Adding loudness analyzer to AudioSynesthesia

diff --git a/Engine/Plugins/Runtime/AudioSynesthesia/Source/AudioSynesthesia/Classes/LoudnessNRT.h b/Engine/Plugins/Runtime/AudioSynesthesia/Source/AudioSynesthesia/Classes/LoudnessNRT.h
new file mode 100644
--- /dev/null
+++ b/Engine/Plugins/Runtime/AudioSynesthesia/Source/AudioSynesthesia/Classes/LoudnessNRT.h
@@ -0,0 +1,117 @@
+// Copyright 1998-2019 Epic Games, Inc. All Rights Reserved.
+
+#pragma once
+
+#include "AudioSynesthesiaNRT.h"
+#include "LoudnessNRT.generated.h"
+
+/** ELoudnessNRTCurveTypeEnum
+ *
+ * Enumeration of available equal loudness curves. Loudness curves based on IEC 61672-1:2013 Electroacoustics - Sound level meters - Part 1: Specifications.
+ */
+UENUM(BlueprintType)
+enum class ELoudnessNRTCurveTypeEnum : uint8
+{
+	/** Loudness Curve A Weighting. Commonly used in environmental sound measurements. Best for low level sounds. */
+	A 		UMETA(DisplayName="A"),
+	/** Loudness Curve B Weighting. Relative to "A", gives more precedence to frequencies below 1kHz. */
+	B		UMETA(DisplayName="B"),
+	/** Loudness Curve C Weighting. Relative to "A" and "B", gives more precedence to frequencies below 1kHz. */
+	C		UMETA(DisplayName="C"),
+	/** Loudness Curve D Weighting. Simlar to "B" but with an emphasis on presence in the 2kHz-6KHz frequency range. */
+	D		UMETA(DisplayName="D"),
+	/** No loudness curve weighting. */
+	None	UMETA(DisplayName="None"),
+};
+
+/** ULoudnessNRTSettings
+ *
+ * Settings for a ULoudnessNRT analyzer.
+ */
+UCLASS(Blueprintable)
+class AUDIOSYNESTHESIA_API ULoudnessNRTSettings : public UAudioSynesthesiaNRTSettings
+{
+	GENERATED_BODY()
+	public:
+
+		ULoudnessNRTSettings();
+
+		/** Number of seconds between loudness measurements */
+		UPROPERTY(EditAnywhere, BlueprintReadWrite, Category=AudioAnalyzer, meta = (ClampMin = "0.01", ClampMax = "0.25"))
+		float AnalysisPeriod;
+
+		/** Minimum analysis frequency for calculating loudness. */
+		UPROPERTY(EditAnywhere, BlueprintReadWrite, Category=AudioAnalyzer, meta = (ClampMin = "20.0", ClampMax = "20000"))
+		float MinimumFrequency;
+
+		/** Maximum analysis frequency for calculating loudness. */
+		UPROPERTY(EditAnywhere, BlueprintReadWrite, Category=AudioAnalyzer, meta = (ClampMin = "20.0", ClampMax = "20000"))
+		float MaximumFrequency;
+
+		/** Type of equal loudness curve used in calculations */
+		UPROPERTY(EditAnywhere, BlueprintReadWrite, Category=AudioAnalyzer)
+		ELoudnessNRTCurveTypeEnum CurveType;
+
+		/** dB level to denote silence.  Used when calculating normalized loudness. */
+		UPROPERTY(EditAnywhere, BlueprintReadWrite, AdvancedDisplay, Category=AudioAnalyzer, meta = (ClampMin = "-100.0", ClampMax="0"))
+		float NoiseFloorDb;
+
+		/** Convert ULoudnessNRTSettings to FLoudnessNRTSettings */
+		TUniquePtr<Audio::IAnalyzerNRTSettings> GetSettings();
+};
+
+
+/** ULoudnessNRT
+ *
+ * ULoudnessNRT calculates the temporal evolution of perceptual loudness for a given 
+ * sound. Loudness is available for individual channels or the overall sound asset. Normalized 
+ * loudness values convert the range to 0.0 to 1.0 where 0.0 is the noise floor and 1.0 is the
+ * maximum loudness of the particular sound.
+ */
+UCLASS(Blueprintable)
+class AUDIOSYNESTHESIA_API ULoudnessNRT : public UAudioSynesthesiaNRT
+{
+	GENERATED_BODY()
+
+	public:
+
+		ULoudnessNRT();
+
+		/** The settings for the audio analyzer.  */
+		UPROPERTY(EditAnywhere, BlueprintReadWrite, Category=AudioAnalyzer)
+		ULoudnessNRTSettings* Settings;
+
+		/** Get the overall loudness of the analyzed sound at a given time. */
+		UFUNCTION(BlueprintCallable, Category="Audio Analyzer")
+		void GetLoudnessAtTime(const float InSeconds, float& OutLoudness);
+
+		/** Get a specific channel loudness of the analyzed sound at a given time. */
+		UFUNCTION(BlueprintCallable, Category="Audio Analyzer")
+		void GetChannelLoudnessAtTime(const float InSeconds, const int32 InChannel, float& OutLoudness);
+
+		/**
+		 * Get the normalized overall loudness of the analyzed sound at a given time. Normalized loudness
+		 * is always between 0.0 to 1.0. 0.0 refers to the noise floor while 1.0 refers to the maximum 
+		 * loudness in the sound.
+		 */
+		UFUNCTION(BlueprintCallable, Category="Audio Analyzer")
+		void GetNormalizedLoudnessAtTime(const float InSeconds, float& OutLoudness);
+
+		/**
+		 * Get a specific channel normalized loudness of the analyzed sound at a given time. Normalized 
+		 * loudness is always between 0.0 to 1.0. 0.0 refers to the noise floor while 1.0 refers to the 
+		 * maximum loudness in the sound.
+		 */
+		UFUNCTION(BlueprintCallable, Category="Audio Analyzer")
+		void GetNormalizedChannelLoudnessAtTime(const float InSeconds, const int32 InChannel, float& OutLoudness);
+
+		/** Convert ULoudnessNRTSettings to FLoudnessNRTSettings */
+ 		TUniquePtr<Audio::IAnalyzerNRTSettings> GetSettings() override;
+
+	protected:
+
+		/** Return the name of the IAudioAnalyzerNRTFactory associated with this UAudioAnalyzerNRT */
+		FName GetAnalyzerNRTFactoryName() const override;
+
+};
+
diff --git a/Engine/Plugins/Runtime/AudioSynesthesia/Source/AudioSynesthesia/Private/AudioSynesthesiaModule.cpp b/Engine/Plugins/Runtime/AudioSynesthesia/Source/AudioSynesthesia/Private/AudioSynesthesiaModule.cpp
--- a/Engine/Plugins/Runtime/AudioSynesthesia/Source/AudioSynesthesia/Private/AudioSynesthesiaModule.cpp
+++ b/Engine/Plugins/Runtime/AudioSynesthesia/Source/AudioSynesthesia/Private/AudioSynesthesiaModule.cpp
@@ -4,6 +4,7 @@
 #include "CoreMinimal.h"
 #include "Modules/ModuleInterface.h"
 #include "Features/IModularFeatures.h"
+#include "LoudnessNRTFactory.h"
 
 DEFINE_LOG_CATEGORY(LogAudioSynesthesia);
 
@@ -13,9 +14,13 @@ namespace Audio
 {
 	void FAudioSynesthesiaModule::StartupModule()
 	{
+		// Register factories on startup
+		IModularFeatures::Get().RegisterModularFeature(FLoudnessNRTFactory::GetModularFeatureName(), &LoudnessFactory);
 	}
 
 	void FAudioSynesthesiaModule::ShutdownModule()
 	{
+		// Unregister factories on shutdown
+		IModularFeatures::Get().UnregisterModularFeature(FLoudnessNRTFactory::GetModularFeatureName(), &LoudnessFactory);
 	}
 }
diff --git a/Engine/Plugins/Runtime/AudioSynesthesia/Source/AudioSynesthesia/Private/LoudnessAnalyzer.cpp b/Engine/Plugins/Runtime/AudioSynesthesia/Source/AudioSynesthesia/Private/LoudnessAnalyzer.cpp
new file mode 100644
--- /dev/null
+++ b/Engine/Plugins/Runtime/AudioSynesthesia/Source/AudioSynesthesia/Private/LoudnessAnalyzer.cpp
@@ -0,0 +1,271 @@
+// Copyright 1998-2019 Epic Games, Inc. All Rights Reserved.
+
+
+#include "LoudnessAnalyzer.h"
+#include "Audio.h"
+#include "DSP/BufferVectorOperations.h"
+#include "DSP/DeinterleaveView.h"
+
+namespace 
+{
+	static const float Log10Scale = 1.0f / FMath::Loge(10.f);
+}
+
+namespace Audio
+{
+    /* Equal loudness curves from: ANSI Standards S1.4-1983 and S1.42-2001 */
+    /* https://en.wikipedia.org/wiki/A-weighting */
+    float GetEqualLoudnessAWeightForFrequency(const float Freq)
+    {
+        const float FreqSqrd = Freq * Freq;
+
+        return 1.258925411794167f * ((148840000.f * FreqSqrd * FreqSqrd) / ((FreqSqrd + 424.36f) * (FreqSqrd + 148840000.f) * FMath::Sqrt((FreqSqrd + 11599.29f) * (FreqSqrd + 544496.41f))));
+    }
+
+    float GetEqualLoudnessBWeightForFrequency(const float Freq)
+    {
+        const float FreqSqrd = Freq * Freq;
+
+        return 1.019764760044717f * ((148840000.f * Freq * FreqSqrd) / ((FreqSqrd + 424.36f) * (FreqSqrd + 148840000.f) * FMath::Sqrt(FreqSqrd + 25122.25f)));
+    }
+
+    float GetEqualLoudnessCWeightForFrequency(const float Freq)
+    {
+        const float FreqSqrd = Freq * Freq;
+
+        return 1.006931668851804f * ((148840000.f * FreqSqrd) / ((FreqSqrd + 424.36f) * (FreqSqrd + 148840000.f)));
+    }
+
+    float GetEqualLoudnessDWeightForFrequency(const float Freq)
+    {
+        const float FreqSqrd = Freq * Freq;
+        const float H1 = (1037918.48f - FreqSqrd);
+        const float H2 = (9837328.f - FreqSqrd);
+        const float H = (H1 + 1080768.16f * FreqSqrd) / (H2 + 11723776.f * FreqSqrd);
+		return (Freq / 6.8966888496476e-5f) * FMath::Sqrt(H / ((FreqSqrd + 79919.29f) * (FreqSqrd + 1345600.f)));
+    }
+
+    float GetEqualLoudnessNoneWeightForFrequency(const float Freq)
+    {
+        return 1.0f;
+    }
+    
+
+	/****************************************************************************/
+	/***********************    FLoudnessAnalyzer    ****************************/
+	/****************************************************************************/
+
+    FLoudnessAnalyzer::FLoudnessAnalyzer(float InSampleRate, const FLoudnessAnalyzerSettings& InSettings)
+    :   Settings(InSettings)
+    ,   EnergyScale(1.0f)
+    ,   MinFreqIndex(0)
+    ,   MaxFreqIndex(0)
+    ,   FFTFreqSize((Settings.FFTSize / 2) + 1)
+    ,   Window(Settings.WindowType, Settings.FFTSize, 1, false)
+    {
+		checkf(FMath::IsPowerOfTwo(Settings.FFTSize), TEXT("FFT size must be a power of two"));
+		checkf(Settings.FFTSize >= 4, TEXT("FFT size must be atleast 4"))
+        check(InSampleRate > 0);
+
+        // Determine which fft bins will be used for calculating loudness
+        const float FreqPerBin = InSampleRate / Settings.FFTSize;
+        MinFreqIndex = static_cast<int32>(FMath::RoundHalfFromZero(Settings.MinAnalysisFrequency / FreqPerBin));
+        MinFreqIndex = FMath::Max(0, MinFreqIndex);
+        MaxFreqIndex = static_cast<int32>(FMath::RoundHalfFromZero(Settings.MaxAnalysisFrequency / FreqPerBin)) + 1;
+        MaxFreqIndex = FMath::Min(MaxFreqIndex, FFTFreqSize);
+
+        // Resize internal buffers
+        WindowedSamples.Reset(Settings.FFTSize);
+        WindowedSamples.AddUninitialized(Settings.FFTSize);
+
+        FFTOutputRealData.Reset(Settings.FFTSize);
+        FFTOutputRealData.AddUninitialized(Settings.FFTSize);
+
+        FFTOutputImagData.Reset(Settings.FFTSize);
+        FFTOutputImagData.AddUninitialized(Settings.FFTSize);
+
+        CurveWeights.Reset(FFTFreqSize);
+        CurveWeights.AddUninitialized(FFTFreqSize);
+
+        PowerSpectrumBuffer.Reset(FFTFreqSize);
+        PowerSpectrumBuffer.AddUninitialized(FFTFreqSize);
+
+        float(*LoudnessMapFn)(const float);
+        switch (Settings.LoudnessCurveType)
+        {
+            case ELoudnessCurveType::A:
+                LoudnessMapFn = GetEqualLoudnessAWeightForFrequency;
+                break;
+
+            case ELoudnessCurveType::B:
+                LoudnessMapFn = GetEqualLoudnessBWeightForFrequency;
+                break;
+
+            case ELoudnessCurveType::C:
+                LoudnessMapFn = GetEqualLoudnessCWeightForFrequency;
+                break;
+
+            case ELoudnessCurveType::D:
+                LoudnessMapFn = GetEqualLoudnessDWeightForFrequency;
+                break;
+
+            case ELoudnessCurveType::None:
+            default:
+                LoudnessMapFn = GetEqualLoudnessNoneWeightForFrequency;
+                break;
+        }
+
+        /* Calculate equal loudness weighting curve */
+
+        // Convention is to normalize 1kHz to have a loudness of 1.
+        float EqualLoudnessNorm = 1.0f / LoudnessMapFn(1000.0f);
+
+        float* CurveWeightsPtr = CurveWeights.GetData();
+
+        const float FFTSize = static_cast<float>(Settings.FFTSize);
+
+        for (int32 i = 0; i < FFTFreqSize; i++)
+        {
+            const float Freq = static_cast<float>(i) * InSampleRate / FFTSize;
+			float Weighting = EqualLoudnessNorm;
+
+			switch (Settings.LoudnessCurveType)
+			{
+				case ELoudnessCurveType::A:
+					Weighting *= GetEqualLoudnessAWeightForFrequency(Freq);
+					break;
+
+				case ELoudnessCurveType::B:
+					Weighting *= GetEqualLoudnessBWeightForFrequency(Freq);
+					break;
+
+				case ELoudnessCurveType::C:
+					Weighting *= GetEqualLoudnessCWeightForFrequency(Freq);
+					break;
+
+				case ELoudnessCurveType::D:
+					Weighting *= GetEqualLoudnessDWeightForFrequency(Freq);
+					break;
+
+				case ELoudnessCurveType::None:
+				default:
+					Weighting *= GetEqualLoudnessNoneWeightForFrequency(Freq);
+					break;
+			}
+
+            // Curve designed for magnitude domain, but applied to power spectrum. The curve is squared to be applied to power spectrum.
+            CurveWeightsPtr[i] = Weighting * Weighting;
+        }
+
+		// Normalize by FFT Scaling (1 / FFTSize) and by number of samples (1 / WindowSize) to calculate mean squared value.
+		EnergyScale = 1.f / static_cast<float>(Settings.FFTSize * Settings.FFTSize);
+    }
+
+    const FLoudnessAnalyzerSettings& FLoudnessAnalyzer::GetSettings() const
+    {
+        return Settings;
+    }
+
+    float FLoudnessAnalyzer::CalculatePerceptualEnergy(TArrayView<const float> InView)
+    {
+		check(InView.Num() == Settings.FFTSize);
+
+		// Copy input samples and apply window
+        FMemory::Memcpy(WindowedSamples.GetData(), InView.GetData(), sizeof(float) * Settings.FFTSize);
+        Window.ApplyToBuffer(WindowedSamples.GetData());
+
+		// Take FFT of windowed samples
+        const FFTTimeDomainData TimeData = {WindowedSamples.GetData(), Settings.FFTSize};
+        FFTFreqDomainData FreqData = {FFTOutputRealData.GetData(), FFTOutputImagData.GetData()};
+
+        PerformFFT(TimeData, FreqData);
+
+		// Accumulate energy as weighted sum of power spectrum in frequency analysis range.
+        const float* FFTOutputRealPtr = FFTOutputRealData.GetData();
+        const float* FFTOutputImagPtr = FFTOutputImagData.GetData();
+        float* PowerSpectrumPtr = PowerSpectrumBuffer.GetData();
+		const float* CurveWeightsPtr = CurveWeights.GetData();
+
+		float Total = 0.0f;
+        for (int32 i = MinFreqIndex; i < MaxFreqIndex; i++)
+        {
+            PowerSpectrumPtr[i] = FFTOutputRealPtr[i] * FFTOutputRealPtr[i] + FFTOutputImagPtr[i] * FFTOutputImagPtr[i];
+			Total += PowerSpectrumPtr[i] * CurveWeightsPtr[i];
+        }
+
+		// Normalize by FFT Scaling and by number of samples to calculate mean squared value.
+		Total *= EnergyScale;
+
+        return Total;
+    }
+
+    float FLoudnessAnalyzer::CalculateLoudness(TArrayView<const float> InView)
+    {
+        float Energy = CalculatePerceptualEnergy(InView);
+		return ConvertPerceptualEnergyToLoudness(Energy);
+    }
+
+	float FLoudnessAnalyzer::ConvertPerceptualEnergyToLoudness(float InPerceptualEnergy)
+	{
+        return 10.0f * FMath::Loge(InPerceptualEnergy) * Log10Scale;
+	}
+
+	/****************************************************************************/
+	/****************    FMultichannelLoudnessAnalyzer    ***********************/
+	/****************************************************************************/
+
+    FMultichannelLoudnessAnalyzer::FMultichannelLoudnessAnalyzer(float InSampleRate, const FLoudnessAnalyzerSettings& InSettings)
+    :   Analyzer(InSampleRate, InSettings)
+    {
+        // Initialize all channel weights to default to 1.0f
+        ChannelWeights.Init(1.0f, EAudioMixerChannel::ChannelTypeCount);
+
+        // If elevation angle is less than 30 degrees and azimuth is between 60 and 120, set channel weight to 1.5f. 
+        ChannelWeights[EAudioMixerChannel::BackLeft] = 1.5f;
+        ChannelWeights[EAudioMixerChannel::BackRight] = 1.5f;
+
+        MonoBuffer.Reset(InSettings.FFTSize);
+        MonoBuffer.AddUninitialized(InSettings.FFTSize);
+    }
+
+    float FMultichannelLoudnessAnalyzer::CalculatePerceptualEnergy(TArrayView<const float> InView, const int32 InNumChannels, TArray<float>& OutChannelEnergies)
+    {
+		check(InView.Num() == (InNumChannels * Analyzer.GetSettings().FFTSize));
+
+        OutChannelEnergies.Reset(InNumChannels);
+        OutChannelEnergies.AddUninitialized(InNumChannels);
+
+        // Deinterleave audio and send to mono loudness analyzer.
+		TAutoDeinterleaveView<float, FAudioBufferAlignedAllocator> DeinterleaveView(InView, MonoBuffer, InNumChannels);
+		for (TAutoDeinterleaveView<float>::TChannel<FAudioBufferAlignedAllocator> Channel : DeinterleaveView)
+		{
+            OutChannelEnergies[Channel.ChannelIndex] = Analyzer.CalculatePerceptualEnergy(Channel.Values);
+		}
+				
+        // Combine channel energies into overall perceptual energy
+        float Energy = 0.0f;
+        for (int32 ChannelNum = 0; ChannelNum < InNumChannels; ChannelNum++)
+        {
+            Energy += ChannelWeights[ChannelNum] * OutChannelEnergies[ChannelNum];
+        }
+
+        return Energy;
+    }
+
+    float FMultichannelLoudnessAnalyzer::CalculateLoudness(TArrayView<const float> InView, const int32 InNumChannels, TArray<float>& OutChannelLoudness)
+    {
+        // Convert perceptual energy to loudness 
+        float Energy = CalculatePerceptualEnergy(InView, InNumChannels, OutChannelLoudness);
+        for (float& Value : OutChannelLoudness)
+        {
+            Value = FLoudnessAnalyzer::ConvertPerceptualEnergyToLoudness(Value);
+        }
+        return FLoudnessAnalyzer::ConvertPerceptualEnergyToLoudness(Energy);
+    }
+
+    const FLoudnessAnalyzerSettings& FMultichannelLoudnessAnalyzer::GetSettings() const
+    {
+        return Analyzer.GetSettings();
+    }
+}
+
diff --git a/Engine/Plugins/Runtime/AudioSynesthesia/Source/AudioSynesthesia/Private/LoudnessNRT.cpp b/Engine/Plugins/Runtime/AudioSynesthesia/Source/AudioSynesthesia/Private/LoudnessNRT.cpp
new file mode 100644
--- /dev/null
+++ b/Engine/Plugins/Runtime/AudioSynesthesia/Source/AudioSynesthesia/Private/LoudnessNRT.cpp
@@ -0,0 +1,161 @@
+// Copyright 1998-2019 Epic Games, Inc. All Rights Reserved.
+
+#include "LoudnessNRT.h"
+#include "LoudnessNRTFactory.h"
+#include "InterpolateSorted.h"
+
+namespace
+{
+	float InterpolateLoudness(TArrayView<const Audio::FLoudnessDatum> InLoudnessArray, const float InTimestamp)
+	{
+		int32 LowerIndex = INDEX_NONE;
+		int32 UpperIndex = INDEX_NONE;
+		float Alpha = 0.f;
+
+		GetInterpolationParametersAtTimestamp(InLoudnessArray, InTimestamp, LowerIndex, UpperIndex, Alpha);
+
+		if ((INDEX_NONE != LowerIndex) && (INDEX_NONE != UpperIndex))
+		{
+			return FMath::Lerp(InLoudnessArray[LowerIndex].Loudness, InLoudnessArray[UpperIndex].Loudness, Alpha);
+		}
+
+		return 0.f;
+	}
+}
+
+
+/***************************************************************************/
+/**********************    ULoudnessNRTSettings     ************************/
+/***************************************************************************/
+
+ULoudnessNRTSettings::ULoudnessNRTSettings()
+:	AnalysisPeriod(0.01f)	
+,	MinimumFrequency(20.f)
+,	MaximumFrequency(20000.f)
+,	CurveType(ELoudnessNRTCurveTypeEnum::D)
+,	NoiseFloorDb(-60.f)
+{}
+
+TUniquePtr<Audio::IAnalyzerNRTSettings> ULoudnessNRTSettings::GetSettings()
+{
+	TUniquePtr<Audio::FLoudnessNRTSettings> Settings = MakeUnique<Audio::FLoudnessNRTSettings>();
+
+	Settings->AnalysisPeriod = AnalysisPeriod;
+	Settings->MinAnalysisFrequency = MinimumFrequency;
+	Settings->MaxAnalysisFrequency = MaximumFrequency;
+
+	switch (CurveType)
+	{
+		case ELoudnessNRTCurveTypeEnum::A:
+			Settings->LoudnessCurveType = Audio::ELoudnessCurveType::A;
+			break;
+
+		case ELoudnessNRTCurveTypeEnum::B:
+			Settings->LoudnessCurveType = Audio::ELoudnessCurveType::B;
+			break;
+			
+		case ELoudnessNRTCurveTypeEnum::C:
+			Settings->LoudnessCurveType = Audio::ELoudnessCurveType::C;
+			break;
+
+		case ELoudnessNRTCurveTypeEnum::D:
+			Settings->LoudnessCurveType = Audio::ELoudnessCurveType::D;
+			break;
+
+		case ELoudnessNRTCurveTypeEnum::None:
+		default:
+			Settings->LoudnessCurveType = Audio::ELoudnessCurveType::None;
+			break;
+	}
+	
+	return Settings;
+}
+
+/***************************************************************************/
+/**********************        ULoudnessNRT         ************************/
+/***************************************************************************/
+
+ULoudnessNRT::ULoudnessNRT()
+{
+	Settings = CreateDefaultSubobject<ULoudnessNRTSettings>(TEXT("DefaultLoudnessNRTSettings"));
+#if WITH_EDITOR
+	// Bind settings to audio analyze so changes to default settings will trigger analysis.
+	Settings->AnalyzeAudioDelegate.BindUObject(this, &UAudioAnalyzerNRT::AnalyzeAudio);
+#endif
+}
+
+void ULoudnessNRT::GetLoudnessAtTime(const float InSeconds, float& OutLoudness)
+{
+	GetChannelLoudnessAtTime(InSeconds, Audio::FLoudnessNRTResult::ChannelIndexOverall, OutLoudness);
+}
+
+void ULoudnessNRT::GetChannelLoudnessAtTime(const float InSeconds, const int32 InChannel, float& OutLoudness)
+{
+	OutLoudness = 0.0f;
+
+	TSharedPtr<Audio::FLoudnessNRTResult, ESPMode::ThreadSafe> LoudnessResult = GetResult<Audio::FLoudnessNRTResult>();
+
+	if (LoudnessResult.IsValid())
+	{
+		if (!LoudnessResult->IsSortedChronologically())
+		{
+			LoudnessResult->SortChronologically();
+		}
+
+		const TArray<Audio::FLoudnessDatum>& LoudnessArray = LoudnessResult->GetChannelLoudnessArray(InChannel);
+
+		OutLoudness = InterpolateLoudness(LoudnessArray, InSeconds);
+	}
+}
+
+void ULoudnessNRT::GetNormalizedLoudnessAtTime(const float InSeconds, float& OutLoudness)
+{
+	GetNormalizedChannelLoudnessAtTime(InSeconds, Audio::FLoudnessNRTResult::ChannelIndexOverall, OutLoudness);
+}
+
+void ULoudnessNRT::GetNormalizedChannelLoudnessAtTime(const float InSeconds, const int32 InChannel, float& OutLoudness)
+{
+	OutLoudness = 0.0f;
+
+	TSharedPtr<Audio::FLoudnessNRTResult, ESPMode::ThreadSafe> LoudnessResult = GetResult<Audio::FLoudnessNRTResult>();
+
+	if (LoudnessResult.IsValid())
+	{
+		if (!LoudnessResult->IsSortedChronologically())
+		{
+			LoudnessResult->SortChronologically();
+		}
+
+		const TArray<Audio::FLoudnessDatum>& LoudnessArray = LoudnessResult->GetChannelLoudnessArray(InChannel);
+
+		OutLoudness = InterpolateLoudness(LoudnessArray, InSeconds);
+
+		// Subtract offset
+		OutLoudness -= Settings->NoiseFloorDb;
+
+		// Divide by range
+		OutLoudness /= FMath::Max(LoudnessResult->GetChannelLoudnessRange(InChannel, Settings->NoiseFloorDb), SMALL_NUMBER);
+
+		// Clamp to desired range in case value was below noise floor
+		OutLoudness = FMath::Clamp(OutLoudness, 0.f, 1.f);
+	}
+}
+
+TUniquePtr<Audio::IAnalyzerNRTSettings> ULoudnessNRT::GetSettings()
+{
+	TUniquePtr<Audio::IAnalyzerNRTSettings> AnalyzerSettings;
+
+	if (Settings)
+	{
+		AnalyzerSettings = Settings->GetSettings();	
+	}
+
+	return AnalyzerSettings;
+}
+
+FName ULoudnessNRT::GetAnalyzerNRTFactoryName() const
+{
+	static const FName FactoryName(TEXT("LoudnessNRTFactory"));
+	return FactoryName;
+}
+
diff --git a/Engine/Plugins/Runtime/AudioSynesthesia/Source/AudioSynesthesia/Private/LoudnessNRTFactory.cpp b/Engine/Plugins/Runtime/AudioSynesthesia/Source/AudioSynesthesia/Private/LoudnessNRTFactory.cpp
new file mode 100644
--- /dev/null
+++ b/Engine/Plugins/Runtime/AudioSynesthesia/Source/AudioSynesthesia/Private/LoudnessNRTFactory.cpp
@@ -0,0 +1,249 @@
+// Copyright 1998-2019 Epic Games, Inc. All Rights Reserved.
+
+#include "LoudnessNRTFactory.h"
+#include "DSP/SlidingWindow.h"
+
+namespace Audio
+{
+	/************************************************************************/
+	/********************** FLoudnessNRTResult ******************************/
+	/************************************************************************/
+	FArchive &operator <<(FArchive& Ar, FLoudnessDatum& Datum)
+	{
+		Ar << Datum.Channel;
+		Ar << Datum.Timestamp;
+		Ar << Datum.Energy;
+		Ar << Datum.Loudness;
+
+		return Ar;
+	}
+
+	// Use for overall loudness storage
+	const int32 FLoudnessNRTResult::ChannelIndexOverall = -1;
+
+
+	FLoudnessNRTResult::FLoudnessNRTResult()
+	: bIsSortedChronologically(false)
+	{
+		// By  default, overall loudness exists
+		ChannelLoudnessArrays.Add(ChannelIndexOverall);
+		ChannelLoudnessIntervals.Add(ChannelIndexOverall);
+	}
+
+
+	void FLoudnessNRTResult::Serialize(FArchive& Archive)
+	{
+		Archive << bIsSortedChronologically;
+		Archive << ChannelLoudnessArrays;
+		Archive << ChannelLoudnessIntervals;
+	}
+
+	void FLoudnessNRTResult::Add(const FLoudnessDatum& InDatum)
+	{
+		// Store loudness data in appropriate channel
+		TArray<FLoudnessDatum>& LoudnessArray = ChannelLoudnessArrays.FindOrAdd(InDatum.Channel);
+		LoudnessArray.Add(InDatum);
+
+		// Update intervals per channel
+		FFloatInterval& Interval = ChannelLoudnessIntervals.FindOrAdd(InDatum.Channel);
+		Interval.Include(InDatum.Loudness);
+
+		// Mark as not sorted
+		bIsSortedChronologically = false;
+	}
+
+	const TArray<FLoudnessDatum>& FLoudnessNRTResult::GetChannelLoudnessArray(int32 ChannelIdx) const
+	{
+		return ChannelLoudnessArrays[ChannelIdx];
+	}
+
+	const TArray<FLoudnessDatum>& FLoudnessNRTResult::GetLoudnessArray() const
+	{
+		return ChannelLoudnessArrays[ChannelIndexOverall];
+	}
+
+	float FLoudnessNRTResult::GetLoudnessRange(float InNoiseFloor) const
+	{
+		return GetChannelLoudnessRange(ChannelIndexOverall, InNoiseFloor);
+	}
+
+	float FLoudnessNRTResult::GetChannelLoudnessRange(int32 InChannelIdx, float InNoiseFloor) const
+	{
+		const FFloatInterval& Interval = ChannelLoudnessIntervals.FindRef(InChannelIdx);
+
+		if (Interval.Contains(InNoiseFloor))
+		{
+			// Noise floor is greater than minimum, so range must be relative to noise floor
+			return Interval.Max - InNoiseFloor;
+		}
+		else if (InNoiseFloor > Interval.Max)
+		{
+			return 0.f;
+		}
+
+		// Noise floor is less than minimum, so use entire range.
+		return Interval.Size();
+	}
+
+	int32 FLoudnessNRTResult::GetNumChannels() const 
+	{
+		return FMath::Max(0, ChannelLoudnessArrays.Num() - 1);
+	}
+
+	bool FLoudnessNRTResult::IsSortedChronologically() const
+	{
+		return bIsSortedChronologically;
+	}
+
+	void FLoudnessNRTResult::SortChronologically()
+	{
+		for (auto& KeyValue : ChannelLoudnessArrays)
+		{
+			KeyValue.Value.Sort([](const FLoudnessDatum& A, const FLoudnessDatum& B) { return A.Timestamp < B.Timestamp; });
+		}
+
+		bIsSortedChronologically = true;
+	}
+
+	/************************************************************************/
+	/********************** FLoudnessNRTWorker ******************************/
+	/************************************************************************/
+
+	FLoudnessNRTWorker::FLoudnessNRTWorker(const FAnalyzerNRTParameters& InParams, const FLoudnessNRTSettings& InAnalyzerSettings)
+	: NumChannels(InParams.NumChannels)
+	, NumAnalyzedBuffers(0)
+	, NumHopFrames(0)
+	, SampleRate(InParams.SampleRate)
+	{
+		check(NumChannels >= 0);
+
+		// Include NumChannels in calculating window size because windows are generated 
+		// with interleaved samples and deinterleaved later. 
+		int32 NumWindowSamples = InAnalyzerSettings.FFTSize * NumChannels;
+
+		NumHopFrames = FMath::CeilToInt(InAnalyzerSettings.AnalysisPeriod * InParams.SampleRate);
+
+		NumHopFrames = FMath::Max(1, NumHopFrames);
+
+		int32 NumHopSamples = NumHopFrames * NumChannels;
+
+		InternalBuffer = MakeUnique<TSlidingBuffer<float>>(NumWindowSamples, NumHopSamples);
+
+		Analyzer = MakeUnique<FMultichannelLoudnessAnalyzer>(InParams.SampleRate, InAnalyzerSettings);
+	}
+
+	void FLoudnessNRTWorker::Analyze(TArrayView<const float> InAudio, IAnalyzerNRTResult* OutResult) 
+	{
+		// Assume that outer layers ensured that this is of correct type.
+		FLoudnessNRTResult* LoudnessResult = static_cast<FLoudnessNRTResult*>(OutResult);
+
+		check(nullptr != LoudnessResult);
+		
+		TAutoSlidingWindow<float> SlidingWindow(*InternalBuffer, InAudio, InternalWindow);
+
+		// Loop through entire array of input samples. 
+		for (const TArray<float>& Window : SlidingWindow)
+		{
+			AnalyzeWindow(Window, *LoudnessResult);
+		}
+	}
+
+	// Called when analysis of audio asset is complete. 
+	void FLoudnessNRTWorker::Finalize(IAnalyzerNRTResult* OutResult) 
+	{
+		FLoudnessNRTResult* LoudnessResult = static_cast<FLoudnessNRTResult*>(OutResult);
+
+		check(nullptr != LoudnessResult);
+
+		TArray<float> EmptyArray;
+		bool bDoFlush = true; // Add zeros to last window until all data has been processed.
+		TAutoSlidingWindow<float> SlidingWindow(*InternalBuffer, EmptyArray, InternalWindow, bDoFlush);
+
+		// Loop through entire array of input samples. 
+		for (const TArray<float>& Window : SlidingWindow)
+		{
+			AnalyzeWindow(Window, *LoudnessResult);
+		}
+
+		// Prepare result for fast lookup.
+		LoudnessResult->SortChronologically(); 
+
+		// Reset internal counters
+		NumAnalyzedBuffers = 0; 
+
+		// Reset sliding buffer
+		InternalBuffer->Reset();
+	}
+
+	void FLoudnessNRTWorker::AnalyzeWindow(TArrayView<const float> InWindow, FLoudnessNRTResult& OutResult)
+	{
+		TArray<float> ChannelPerceptualEnergy;
+		float OverallPerceptualEnergy = 0.0f;
+
+		// Calculate perceptual energy
+		OverallPerceptualEnergy = Analyzer->CalculatePerceptualEnergy(InWindow, NumChannels, ChannelPerceptualEnergy);
+
+		// Convert to perceptual energy to loudness
+		TArray<float> ChannelLoudness;
+		float OverallLoudness = 0.f;
+		
+		OverallLoudness = FLoudnessAnalyzer::ConvertPerceptualEnergyToLoudness(OverallPerceptualEnergy);
+
+		for (float PerceptualEnergy : ChannelPerceptualEnergy)
+		{
+			ChannelLoudness.Add(FLoudnessAnalyzer::ConvertPerceptualEnergyToLoudness(PerceptualEnergy));
+		}
+
+		// Calculate timestamp
+		float Timestamp = ((NumAnalyzedBuffers * NumHopFrames) + (0.5f * Analyzer->GetSettings().FFTSize)) / SampleRate;
+
+		// Add overall FLoudnessDatum to result
+		const FLoudnessDatum OverallDatum = {FLoudnessNRTResult::ChannelIndexOverall, Timestamp, OverallPerceptualEnergy, OverallLoudness};
+		OutResult.Add(OverallDatum);
+
+		// Add channel FLoudnessDatum to result
+		for (int32 ChannelIdx = 0; ChannelIdx < NumChannels; ChannelIdx++)
+		{
+			const FLoudnessDatum ChannelDatum = {ChannelIdx, Timestamp, ChannelPerceptualEnergy[ChannelIdx], ChannelLoudness[ChannelIdx]};
+
+			OutResult.Add(ChannelDatum);
+		}	
+
+		// update counters
+		NumAnalyzedBuffers++;
+	}
+
+
+	/************************************************************************/
+	/************************* FLoudnessNRTFactory **************************/
+	/************************************************************************/
+
+	/** Name of specific analyzer type. */
+	FName FLoudnessNRTFactory::GetName() const 
+	{
+		static FName FactoryName(TEXT("LoudnessNRTFactory"));
+		return FactoryName;
+	}
+
+	/** Human readable name of analyzer. */
+	FString FLoudnessNRTFactory::GetTitle() const
+	{
+		return TEXT("Loudness Analyzer Non-Real-Time");
+	}
+
+	TUniquePtr<IAnalyzerNRTResult> FLoudnessNRTFactory::NewResult()
+	{
+		TUniquePtr<FLoudnessNRTResult> Result = MakeUnique<FLoudnessNRTResult>();
+		return Result;
+	}
+
+	TUniquePtr<IAnalyzerNRTWorker> FLoudnessNRTFactory::NewWorker(const FAnalyzerNRTParameters& InParams, const IAnalyzerNRTSettings* InSettings)
+	{
+		const FLoudnessNRTSettings* LoudnessSettings = static_cast<const FLoudnessNRTSettings*>(InSettings);
+
+		check(nullptr != LoudnessSettings);
+
+		return MakeUnique<FLoudnessNRTWorker>(InParams, *LoudnessSettings);
+	}
+}
+
diff --git a/Engine/Plugins/Runtime/AudioSynesthesia/Source/AudioSynesthesia/Public/AudioSynesthesiaModule.h b/Engine/Plugins/Runtime/AudioSynesthesia/Source/AudioSynesthesia/Public/AudioSynesthesiaModule.h
--- a/Engine/Plugins/Runtime/AudioSynesthesia/Source/AudioSynesthesia/Public/AudioSynesthesiaModule.h
+++ b/Engine/Plugins/Runtime/AudioSynesthesia/Source/AudioSynesthesia/Public/AudioSynesthesiaModule.h
@@ -5,6 +5,7 @@
 #include "CoreMinimal.h"
 #include "Modules/ModuleInterface.h"
 #include "Modules/ModuleManager.h"
+#include "LoudnessNRTFactory.h"
 
 DECLARE_LOG_CATEGORY_EXTERN(LogAudioSynesthesia, Log, All);
 
@@ -19,5 +20,6 @@ namespace Audio
 		virtual void ShutdownModule() override;
 
 	private:
+		FLoudnessNRTFactory LoudnessFactory;
 	};
 }
diff --git a/Engine/Plugins/Runtime/AudioSynesthesia/Source/AudioSynesthesia/Public/InterpolateSorted.h b/Engine/Plugins/Runtime/AudioSynesthesia/Source/AudioSynesthesia/Public/InterpolateSorted.h
new file mode 100644
--- /dev/null
+++ b/Engine/Plugins/Runtime/AudioSynesthesia/Source/AudioSynesthesia/Public/InterpolateSorted.h
@@ -0,0 +1,48 @@
+// Copyright 1998-2019 Epic Games, Inc. All Rights Reserved.
+
+#pragma once
+
+namespace Audio
+{
+	/** GetInterpolationParametersAtTimestamp
+	 *
+	 *  Calculates the indices and alpha parameter for interpolating between to elements in the array.
+	 *
+	 *  ElementType must be a type with a "float Timestamp;" member.
+	 *  InArray is a array of ElementTypes which are sorted in ascending order via the Timestamp member.
+	 *  InTimestamp is the query timestamp.
+	 *
+	 *  OutLowerIndex is populated with the index of the last element in InArray with a Timestamp less than or equal to InTimestamp. On error, this is set to INDEX_NONE
+	 *  OutUpperIndex is populated with the index of the first element in InArray with a Timestamp greater than InTimestamp. On error, this is set to INDEX_NONE
+	 *  OutAlpha is a number between 0.f and 1.f which describes the blending between the lower and upper indices. 0.f refers OutLowerIndex and 1.f to OutUpperIndex.
+	 */
+	template<typename ElementType>
+	void GetInterpolationParametersAtTimestamp(TArrayView<const ElementType> InArray, float InTimestamp, int32& OutLowerIndex, int32& OutUpperIndex, float& OutAlpha)
+	{
+		if (0 == InArray.Num())
+		{
+			OutLowerIndex = INDEX_NONE;
+			OutUpperIndex = INDEX_NONE;
+			OutAlpha = 1.f;
+			return;
+		}
+		
+		// Get upper and lower array indices with FLoudnessData closest to the query timestamp.
+		OutUpperIndex = Algo::UpperBoundBy(InArray, InTimestamp, [](const ElementType& Datum) { return Datum.Timestamp; });
+		OutLowerIndex = OutUpperIndex - 1;
+
+		// Ensure indices are within range.
+		int32 MaxIdx = InArray.Num() - 1;
+		OutLowerIndex = FMath::Clamp(OutLowerIndex, 0, MaxIdx);
+		OutUpperIndex = FMath::Clamp(OutUpperIndex, 0, MaxIdx);
+
+		const ElementType& LowerDatum = InArray[OutLowerIndex];
+		const ElementType& UpperDatum = InArray[OutUpperIndex];
+
+		// Determine time difference and linear interpolation factor from data.
+		const float TimeDelta = FMath::Max(UpperDatum.Timestamp - LowerDatum.Timestamp, SMALL_NUMBER);
+		OutAlpha = (InTimestamp - LowerDatum.Timestamp) / TimeDelta;
+		OutAlpha = FMath::Clamp(OutAlpha, 0.0f, 1.0f);
+	}
+}
+
diff --git a/Engine/Plugins/Runtime/AudioSynesthesia/Source/AudioSynesthesia/Public/LoudnessAnalyzer.h b/Engine/Plugins/Runtime/AudioSynesthesia/Source/AudioSynesthesia/Public/LoudnessAnalyzer.h
new file mode 100644
--- /dev/null
+++ b/Engine/Plugins/Runtime/AudioSynesthesia/Source/AudioSynesthesia/Public/LoudnessAnalyzer.h
@@ -0,0 +1,172 @@
+// Copyright 1998-2019 Epic Games, Inc. All Rights Reserved.
+
+#pragma once
+
+#include "DSP/BufferVectorOperations.h"
+#include "DSP/AudioFFT.h"
+
+namespace Audio
+{
+    /** Equal loudness curves from: ANSI Standards S1.4-1983 and S1.42-2001 */
+	enum class ELoudnessCurveType : uint8
+	{
+		A,
+		B,
+		C,
+		D,
+		None
+	};
+
+	/** Generate equal loudness weight for magnitude spectrum using curve A */
+    float GetEqualLoudnessAWeigthForFrequency(const float Freq);
+
+	/** Generate equal loudness weight for magnitude spectrum using curve B */
+    float GetEqualLoudnessBWeigthForFrequency(const float Freq);
+
+	/** Generate equal loudness weight for magnitude spectrum using curve C */
+    float GetEqualLoudnessCWeigthForFrequency(const float Freq);
+
+	/** Generate equal loudness weight for magnitude spectrum using curve D */
+    float GetEqualLoudnessDWeigthForFrequency(const float Freq);
+
+	/** Generate equal loudness weight for magnitude spectrum using curve None */
+    float GetEqualLoudnessNoneWeigthForFrequency(const float Freq);
+
+	/** Settings for FLoudnessAnalyzer */
+    struct FLoudnessAnalyzerSettings
+    {
+		/** Size of FFT describe in number of samples */
+        int32 FFTSize;
+
+		/** Type of window to be applied to input audio */
+        EWindowType WindowType;
+
+		/** Type of equal loudness curve to weight power spectrum */
+        ELoudnessCurveType LoudnessCurveType;
+
+		/** Minimum frequency for analyzing audio */
+        float MinAnalysisFrequency;
+
+		/** Maximum frequency for analyzing audio */
+        float MaxAnalysisFrequency;
+
+        /** Default settings */
+		FLoudnessAnalyzerSettings()
+			: FFTSize(8192)
+			, WindowType(EWindowType::Blackman)
+			, LoudnessCurveType(ELoudnessCurveType::D)
+			, MinAnalysisFrequency(20.0f)
+			, MaxAnalysisFrequency(10000.0f)
+		{}
+    };
+
+	/** FLoudnessAnalyzer
+	 *
+	 * FLoudnessAnalyzer will calculate perceptual energy or loudness one window at a time. 
+	 * This analyzer only handles mono audio. See FMultichannelLoudnessAnalyzer for multichannel
+	 * support.
+	 */
+    class FLoudnessAnalyzer
+    {
+        public:
+			/** Construct analyzer */
+            FLoudnessAnalyzer(float InSampleRate, const FLoudnessAnalyzerSettings& InSettings);
+
+			/**
+			 * Calculats the perceptual energy for the input samples.  The number of 
+			 * input samples must be equal to the FFTSize supplied in FLoudnessAnalyzerSettings.
+			 */
+            float CalculatePerceptualEnergy(TArrayView<const float> InSampleView);
+
+			/**
+			 * Calculats the loudness for the input samples.  The number of  input samples must 
+			 * be equal to the FFTSize supplied in FLoudnessAnalyzerSettings.
+			 */
+            float CalculateLoudness(TArrayView<const float> InSampleView);
+
+			/**
+			 * Convert from perceptual energy to loudness [Loudness = 10 * log_10(Energy)]
+			 */
+			static float ConvertPerceptualEnergyToLoudness(float InPerceptualEnergy);
+
+			/**
+			 * Return const reference to settings used inside this analyzer.
+			 */
+            const FLoudnessAnalyzerSettings& GetSettings() const;
+        private:
+
+            FLoudnessAnalyzerSettings Settings;
+
+			/** Scale factor used to normalize between differing window sizes and fft sizes */
+            float EnergyScale;
+
+            int32 MinFreqIndex;
+            int32 MaxFreqIndex;
+
+			/** Number of FFTFrequency bins worth looking at. */
+			int32 FFTFreqSize;
+
+            FWindow Window;
+
+            AlignedFloatBuffer WindowedSamples;
+            AlignedFloatBuffer FFTOutputRealData;
+            AlignedFloatBuffer FFTOutputImagData;
+
+            AlignedFloatBuffer CurveWeights;
+            AlignedFloatBuffer PowerSpectrumBuffer;
+    };
+
+
+	/** FMultichannelLoudnessAnalyzer
+	 *
+	 * FMultichannelLoudnessAnalyzer will calculate perceptual energy or loudness one window 
+	 * at a time. This module can handle one or more channels at a time. It expects 
+	 * interleaved audio samples.
+	 */
+    class FMultichannelLoudnessAnalyzer 
+    {
+        public:
+			/** Construct analyzer. */
+            FMultichannelLoudnessAnalyzer(float InSampleRate, const FLoudnessAnalyzerSettings& InSettings);
+
+			/** 
+			 * Calculate perceptual energy. 
+			 *
+			 * InSampleView		  InSampleView is expected to contain interleaved audio with 
+			 *                    (InSettings.FFTSize * InNumChannels) number of samples.
+			 *
+			 * InNumChannels      Number of input channels.
+			 *
+			 * OutChannelEnergies Will be populated with energies on a per channel basis.
+			 *
+			 * The overall perceptual energy is returned.
+			 */
+            float CalculatePerceptualEnergy(TArrayView<const float> InSampleView, const int32 InNumChannels, TArray<float>& OutChannelEnergies); 
+
+			/** 
+			 * Calculate loudness
+			 *
+			 * InSampleView		  InSampleView is expected to contain interleaved audio with 
+			 *                    (InSettings.FFTSize * InNumChannels) number of samples.
+			 *
+			 * InNumChannels      Number of input channels.
+			 *
+			 * OutChannelLoudness Will be populated with energies on a per channel basis.
+			 *
+			 * The overall loudness is returned.
+			 */
+            float CalculateLoudness(TArrayView<const float> InSampleView, const int32 InNumChannels, TArray<float>& OutChannelLoudness);
+
+			/**
+			 * Return const reference to settings used inside this analyzer.
+			 */
+            const FLoudnessAnalyzerSettings& GetSettings() const;
+
+        private:
+
+            AlignedFloatBuffer MonoBuffer;
+            TArray<float> ChannelWeights;
+            FLoudnessAnalyzer Analyzer;
+    };
+}
+
diff --git a/Engine/Plugins/Runtime/AudioSynesthesia/Source/AudioSynesthesia/Public/LoudnessNRTFactory.h b/Engine/Plugins/Runtime/AudioSynesthesia/Source/AudioSynesthesia/Public/LoudnessNRTFactory.h
new file mode 100644
--- /dev/null
+++ b/Engine/Plugins/Runtime/AudioSynesthesia/Source/AudioSynesthesia/Public/LoudnessNRTFactory.h
@@ -0,0 +1,171 @@
+// Copyright 1998-2019 Epic Games, Inc. All Rights Reserved.
+
+#pragma once
+
+#include "LoudnessAnalyzer.h"
+#include "IAudioAnalyzerNRTInterface.h"
+#include "DSP/SlidingWindow.h"
+
+namespace Audio
+{
+	/** FLoudnessNRTSettings
+	 *
+	 * Contains settings for loudness analyzer.
+	 */
+	class FLoudnessNRTSettings : public IAnalyzerNRTSettings, public FLoudnessAnalyzerSettings
+	{
+		public:
+			/** Number of seconds between loudness measurements */
+			float AnalysisPeriod;
+
+			FLoudnessNRTSettings()
+			:	AnalysisPeriod(0.01f)
+			{}
+	};
+
+
+	/** FLoudnessNRTResult
+	 *
+	 * Holds the loundess values per a time step.
+	 */
+	struct FLoudnessDatum
+	{
+		int32 Channel = 0;
+		float Timestamp = 0.f;
+		float Energy = 0.f;
+		float Loudness = 0.f;
+	};
+	/** De/Serialize single loudness datum into archive. */
+	AUDIOSYNESTHESIA_API FArchive &operator <<(FArchive& Ar, FLoudnessDatum& Datum);
+
+	/** FLoudnessNRTResult
+	 *
+	 * FLoudnessNRTResult contains the temporal evolution of loudness.
+	 */
+	class AUDIOSYNESTHESIA_API FLoudnessNRTResult : public IAnalyzerNRTResult
+	{
+	public:
+		/** ChannelIndexOverall
+		 *
+		 * Denotes the overall loudness channel index as opposed individual channel indices.
+		 */
+		static const int32 ChannelIndexOverall;
+
+		FLoudnessNRTResult();
+
+		/**
+		 * Defines how to serialize result.
+		 */
+		void Serialize(FArchive& Archive) override;
+
+		/**
+		 * Appends an FLoudnessDatum to the container.
+		 */
+		void Add(const FLoudnessDatum& InDatum);
+
+		/**
+		 * Returns const reference to FLoudnessDatum array for individual channel.
+		 */
+		const TArray<FLoudnessDatum>& GetChannelLoudnessArray(int32 ChannelIdx) const;
+
+		/**
+		 * Returns const reference to FLoudnessDatum array associated with overall loudness.
+		 */
+		const TArray<FLoudnessDatum>& GetLoudnessArray() const;
+
+		/**
+		 * Returns range in dB of overall loudness result given the noise floor.
+		 */
+		float GetLoudnessRange(float InNoiseFloor) const;
+
+		/**
+		 * Returns range in dB of loudness result given the noise floor.
+		 */
+		float GetChannelLoudnessRange(int32 InChannelIdx, float InNoiseFloor) const;
+
+		/**
+		 * Get Number of Channels
+		 */
+		int32 GetNumChannels() const;
+
+		/**
+		 * Returns true if FLoudnessDatum arrays are sorted in chronologically ascending order via their timestamp.
+		 */
+		bool IsSortedChronologically() const;
+
+		/**
+		 * Sorts FLoudnessDatum arrays in chronologically ascnding order via their timestamp.
+		 */
+		void SortChronologically();
+
+	private:
+		TMap<int32, TArray<FLoudnessDatum> > ChannelLoudnessArrays;
+
+		TMap<int32, FFloatInterval> ChannelLoudnessIntervals;
+
+		bool bIsSortedChronologically;
+	};
+
+	/** FLoudnessNRTWorker
+	 *
+	 * FLoudnessNRTWorker performs loudness analysis on input sample buffers.
+	 */
+	class FLoudnessNRTWorker : public IAnalyzerNRTWorker
+	{
+	public:
+		/** Construct a worker */
+		FLoudnessNRTWorker(const FAnalyzerNRTParameters& InParams, const FLoudnessNRTSettings& InAnalyzerSettings);
+
+		/**
+		 * Analyzes input sample buffer and updates result. 
+		 */
+		void Analyze(TArrayView<const float> InAudio, IAnalyzerNRTResult* OutResult) override;
+
+		/**
+		 * Call when all audio data has been analyzed. 
+		 */
+		void Finalize(IAnalyzerNRTResult* OutResult) override;
+
+	private:
+
+		// Analyze a single window
+		void AnalyzeWindow(TArrayView<const float> InWindow, FLoudnessNRTResult& OutResult);
+
+		int32 NumChannels;
+		int32 NumAnalyzedBuffers;
+		int32 NumHopFrames;
+
+		float SampleRate;
+
+		TArray<float> InternalWindow;
+
+		TUniquePtr<TSlidingBuffer<float>> InternalBuffer;
+
+		TUniquePtr<FMultichannelLoudnessAnalyzer> Analyzer;
+	};
+
+	/** FLoudnessNRTFactory
+	 *
+	 * Defines the LoudnessNRT analyzer and creates related classes.
+	 */
+	class FLoudnessNRTFactory : public IAnalyzerNRTFactory
+	{
+		public:
+
+		/** Name of specific analyzer type. */
+		FName GetName() const override;
+
+		/** Human readable name of analyzer. */
+		FString GetTitle() const override;
+
+		/** Creates a new FLoudnessNRTResult */
+		TUniquePtr<IAnalyzerNRTResult> NewResult() override;
+
+		/** 
+		 * Creates a new FLoudnessNRTWorker. This expects IAnalyzerNRTSettings to be a valid pointer to
+		 * a FLoudnessNRTSettings object.
+		 */
+		TUniquePtr<IAnalyzerNRTWorker> NewWorker(const FAnalyzerNRTParameters& InParams, const IAnalyzerNRTSettings* InSettings) override;
+	};
+}
+
diff --git a/Engine/Source/Runtime/AudioAnalyzer/Private/AudioAnalyzerNRT.cpp b/Engine/Source/Runtime/AudioAnalyzer/Private/AudioAnalyzerNRT.cpp
--- a/Engine/Source/Runtime/AudioAnalyzer/Private/AudioAnalyzerNRT.cpp
+++ b/Engine/Source/Runtime/AudioAnalyzer/Private/AudioAnalyzerNRT.cpp
@@ -219,10 +219,9 @@ void UAudioAnalyzerNRT::Serialize(FArchive& Ar)
 			if (nullptr != Factory)
 			{
 				// Create result and worker from factory
-				TUniquePtr<Audio::IAnalyzerNRTResult> NewResult = Factory->NewResult();
 				{
 					FScopeLock ResultLock(&ResultCriticalSection);
-					Result = TSharedPtr<Audio::IAnalyzerNRTResult, ESPMode::ThreadSafe>(NewResult.Release());
+					Result = Factory->NewResultShared<ESPMode::ThreadSafe>();
 				}
 			}
 		}
diff --git a/Engine/Source/Runtime/AudioAnalyzer/Public/IAudioAnalyzerNRTInterface.h b/Engine/Source/Runtime/AudioAnalyzer/Public/IAudioAnalyzerNRTInterface.h
--- a/Engine/Source/Runtime/AudioAnalyzer/Public/IAudioAnalyzerNRTInterface.h
+++ b/Engine/Source/Runtime/AudioAnalyzer/Public/IAudioAnalyzerNRTInterface.h
@@ -5,7 +5,6 @@
 #pragma once
 
 #include "CoreMinimal.h"
-#include "SampleBuffer.h"
 #include "Features/IModularFeature.h"
 
 namespace Audio
@@ -104,6 +103,15 @@ namespace Audio
 		// Create a new result.
 		virtual TUniquePtr<IAnalyzerNRTResult> NewResult() = 0;
 
+		// Convenience function to create a new shared result by calling NewResult.
+		template<ESPMode Mode = ESPMode::Fast>
+		TSharedPtr<IAnalyzerNRTResult, Mode> NewResultShared()
+		{
+			TUniquePtr<Audio::IAnalyzerNRTResult> Result = NewResult();
+
+			return TSharedPtr<Audio::IAnalyzerNRTResult, Mode>(Result.Release());
+		}
+
 		// Create a new worker.
 		virtual TUniquePtr<IAnalyzerNRTWorker> NewWorker(const FAnalyzerNRTParameters& InParams, const IAnalyzerNRTSettings* InSettings) = 0;
 	};
diff --git a/Engine/Source/Runtime/SignalProcessing/Private/FloatArrayMath.cpp b/Engine/Source/Runtime/SignalProcessing/Private/FloatArrayMath.cpp
--- a/Engine/Source/Runtime/SignalProcessing/Private/FloatArrayMath.cpp
+++ b/Engine/Source/Runtime/SignalProcessing/Private/FloatArrayMath.cpp
@@ -10,10 +10,10 @@ namespace
 
 namespace Audio
 {
-	void ArrayCumulativeSum(const TArray<float>& InData, TArray<float>& OutData)
+	void ArrayCumulativeSum(TArrayView<const float> InView, TArray<float>& OutData)
 	{
 		// Initialize output data
-		int32 Num = InData.Num();
+		int32 Num = InView.Num();
 		OutData.Reset();
 		OutData.AddUninitialized(Num);
 
@@ -23,19 +23,19 @@ namespace Audio
 		}
 
 		float* OutDataPtr = OutData.GetData();
-		const float* InDataPtr = InData.GetData();
+		const float* InViewPtr = InView.GetData();
 
 		// Start summing
-		*OutDataPtr = *InDataPtr++;
+		*OutDataPtr = *InViewPtr++;
 
 		for (int32 i = 1; i < Num; i++)
 		{
-			float Temp = *OutDataPtr++ + *InDataPtr++;
+			float Temp = *OutDataPtr++ + *InViewPtr++;
 			*OutDataPtr = Temp;
 		}
 	}
 
-	void ArrayMeanFilter(const TArray<float>& InData, int32 WindowSize, int32 WindowOrigin, TArray<float>& OutData)
+	void ArrayMeanFilter(TArrayView<const float> InView, int32 WindowSize, int32 WindowOrigin, TArray<float>& OutData)
 	{
 		// a quick but sinful implementation of a mean filter. encourages floating point rounding errors. 
 		check(WindowOrigin < WindowSize);
@@ -43,7 +43,7 @@ namespace Audio
 		check(WindowSize > 0);
 
 		// Initialize output data
-		const int32 Num = InData.Num();
+		const int32 Num = InView.Num();
 		OutData.Reset();
 		OutData.AddUninitialized(Num);
 
@@ -53,11 +53,11 @@ namespace Audio
 		}
 		
 		// Use cumulative sum to avoid multiple summations 
-		// Instead of summing over InData[StartIndex:EndIndex], avoid all that
+		// Instead of summing over InView[StartIndex:EndIndex], avoid all that
 		// calculation by taking difference of cumulative sum at those two points:
 		//  cumsum(X[0:b]) - cumsum(X[0:a]) = sum(X[a:b])
 		TArray<float> SummedData;
-		ArrayCumulativeSum(InData, SummedData);
+		ArrayCumulativeSum(InView, SummedData);
 		const float LastSummedData = SummedData.Last();
 		
 		
@@ -101,7 +101,7 @@ namespace Audio
 		}
 	}
 
-	void ArrayMaxFilter(const TArray<float>& InData, int32 WindowSize, int32 WindowOrigin, TArray<float>& OutData)
+	void ArrayMaxFilter(TArrayView<const float> InView, int32 WindowSize, int32 WindowOrigin, TArray<float>& OutData)
 	{
 		// A reasonable implementation of a max filter for the data we're interested in, though surely not the fastest.
 		check(WindowOrigin < WindowSize);
@@ -112,7 +112,7 @@ namespace Audio
 		int32 EndIndex = StartIndex + WindowSize;
 
 		// Initialize output
-		int32 Num = InData.Num();
+		int32 Num = InView.Num();
 		OutData.Reset();
 		OutData.AddUninitialized(Num);
 
@@ -125,16 +125,16 @@ namespace Audio
 		int32 ActualStartIndex = 0;
 		int32 ActualEndIndex = FMath::Min(EndIndex, Num);
 
-		const float* InDataPtr = InData.GetData();
+		const float* InViewPtr = InView.GetData();
 		float* OutDataPtr = OutData.GetData();
 		int32 MaxIndex = 0;
-		float MaxValue = InData[0];
+		float MaxValue = InView[0];
 
 		for (int32 i = ActualStartIndex; i < ActualEndIndex; i++)
 		{
-			if (InDataPtr[i] > MaxValue)
+			if (InViewPtr[i] > MaxValue)
 			{
-				MaxValue = InDataPtr[i];
+				MaxValue = InViewPtr[i];
 				MaxIndex = i;
 			}		
 		}
@@ -153,23 +153,23 @@ namespace Audio
 			{
 				// We need to evaluate the entire window because the previous maximum value was not in this window.
 				MaxIndex = ActualStartIndex;
-				MaxValue = InDataPtr[MaxIndex];
+				MaxValue = InViewPtr[MaxIndex];
 				for (int32 j = ActualStartIndex + 1; j < ActualEndIndex; j++)
 				{
-					if (InDataPtr[j] > MaxValue)
+					if (InViewPtr[j] > MaxValue)
 					{
 						MaxIndex = j;
-						MaxValue = InDataPtr[MaxIndex];
+						MaxValue = InViewPtr[MaxIndex];
 					}
 				}
 			}
 			else
 			{
 				// We only need to inspect the newest sample because the previous maximum value was in this window.
-				if (InDataPtr[ActualEndIndex - 1] > MaxValue)
+				if (InViewPtr[ActualEndIndex - 1] > MaxValue)
 				{
 					MaxIndex = ActualEndIndex - 1;
-					MaxValue = InDataPtr[MaxIndex];
+					MaxValue = InViewPtr[MaxIndex];
 				}
 			}
 
@@ -180,59 +180,70 @@ namespace Audio
 		}
 	}
 
-	void ArrayGetEuclideanNorm(const TArray<float>& InArray, float& OutEuclideanNorm)
+	void ArrayGetEuclideanNorm(TArrayView<const float> InView, float& OutEuclideanNorm)
 	{
 		// Initialize output.
 		OutEuclideanNorm = 0.0f;
-		const int32 Num = InArray.Num();
-		const float* InArrayData = InArray.GetData();
+		const int32 Num = InView.Num();
+		const float* InViewData = InView.GetData();
 		
 		// Sum it up.
 		for (int32 i = 0; i < Num; i++)
 		{
-			OutEuclideanNorm += InArrayData[i] * InArrayData[i];
+			OutEuclideanNorm += InViewData[i] * InViewData[i];
 		}
 
 		OutEuclideanNorm = FMath::Sqrt(OutEuclideanNorm);
 	}
 
-	void ArrayMultiplyByConstantInPlace(TArray<float>& InArray, float InMultiplier)
+	void ArrayClampInPlace(TArrayView<float> InView, float InMin, float InMax)
 	{
-		const int32 Num = InArray.Num();
-		float* InArrayData = InArray.GetData();
+		const int32 Num = InView.Num();
+		float* Data = InView.GetData();
+
+		for (int32 i = 0; i < Num; i++)
+		{
+			Data[i] = FMath::Clamp(Data[i], InMin, InMax);
+		}
+	}
+
+	void ArrayMultiplyByConstantInPlace(TArrayView<float> InView, float InMultiplier)
+	{
+		const int32 Num = InView.Num();
+		float* InViewData = InView.GetData();
 		for (int32 i = 0; i < Num; i++)
 		{
-			InArrayData[i] *= InMultiplier;
+			InViewData[i] *= InMultiplier;
 		}
 	}
 
-	void ArraySubtractByConstantInPlace(TArray<float>& InArray, float InSubtrahend)
+	void ArraySubtractByConstantInPlace(TArrayView<float> InView, float InSubtrahend)
 	{
-		const int32 Num = InArray.Num();
-		float* InArrayData = InArray.GetData();
+		const int32 Num = InView.Num();
+		float* InViewData = InView.GetData();
 		for (int32 i = 0; i < Num; i++)
 		{
-			InArrayData[i] -= InSubtrahend;
+			InViewData[i] -= InSubtrahend;
 		}
 	}
 
-	void ArrayMagnitudeToDecibelInPlace(TArray<float>& InArray)
+	void ArrayMagnitudeToDecibelInPlace(TArrayView<float> InView)
 	{
-		const int32 Num = InArray.Num();
-		float* InArrayData = InArray.GetData();
+		const int32 Num = InView.Num();
+		float* InViewData = InView.GetData();
 		for (int32 i = 0; i < Num; i++)
 		{
-			InArrayData[i] = 20.f * FMath::Loge(InArrayData[i]) / LOGE10;
+			InViewData[i] = 20.f * FMath::Loge(InViewData[i]) / LOGE10;
 		}
 	}
 
-	void ArrayPowerToDecibelInPlace(TArray<float>& InArray)
+	void ArrayPowerToDecibelInPlace(TArrayView<float> InView)
 	{
-		const int32 Num = InArray.Num();
-		float* InArrayData = InArray.GetData();
+		const int32 Num = InView.Num();
+		float* InViewData = InView.GetData();
 		for (int32 i = 0; i < Num; i++)
 		{
-			InArrayData[i] = 10.f * FMath::Loge(InArrayData[i]) / LOGE10;
+			InViewData[i] = 10.f * FMath::Loge(InViewData[i]) / LOGE10;
 		}
 	}
 
@@ -269,9 +280,9 @@ namespace Audio
 		Kernel[RowIndex].OffsetValues = TArray<float>(OffsetValues.GetData(), OffsetValues.Num());
 	}
 
-	void FContiguousSparse2DKernelTransform::TransformArray(TArrayView<const float> InArray, TArray<float>& OutArray) const
+	void FContiguousSparse2DKernelTransform::TransformArray(TArrayView<const float> InView, TArray<float>& OutArray) const
 	{
-		check(InArray.Num() == NumIn);
+		check(InView.Num() == NumIn);
 
 		// Resize output
 		OutArray.Reset(NumOut);
@@ -280,12 +291,12 @@ namespace Audio
 			OutArray.AddUninitialized(NumOut);
 		}
 
-		TransformArray(InArray.GetData(), OutArray.GetData());
+		TransformArray(InView.GetData(), OutArray.GetData());
 	}
 
-	void FContiguousSparse2DKernelTransform::TransformArray(TArrayView<const float> InArray, AlignedFloatBuffer& OutArray) const
+	void FContiguousSparse2DKernelTransform::TransformArray(TArrayView<const float> InView, AlignedFloatBuffer& OutArray) const
 	{	
-		check(InArray.Num() == NumIn);
+		check(InView.Num() == NumIn);
 
 		// Resize output
 		OutArray.Reset(NumOut);
@@ -294,7 +305,7 @@ namespace Audio
 			OutArray.AddUninitialized(NumOut);
 		}
 
-		TransformArray(InArray.GetData(), OutArray.GetData());
+		TransformArray(InView.GetData(), OutArray.GetData());
 	}
 
 	void FContiguousSparse2DKernelTransform::TransformArray(const float* InArray, float* OutArray) const
diff --git a/Engine/Source/Runtime/SignalProcessing/Public/DSP/FloatArrayMath.h b/Engine/Source/Runtime/SignalProcessing/Public/DSP/FloatArrayMath.h
--- a/Engine/Source/Runtime/SignalProcessing/Public/DSP/FloatArrayMath.h
+++ b/Engine/Source/Runtime/SignalProcessing/Public/DSP/FloatArrayMath.h
@@ -9,53 +9,61 @@ namespace Audio
 {
 	/** Cumulative sum of array.
 	 *
-	 *  InData contains data to be cumulatively summed.
-	 *  OutData contains sum and is same size as InData.
+	 *  InView contains data to be cumulatively summed.
+	 *  OutData contains sum and is same size as InView.
 	 */
-	SIGNALPROCESSING_API void ArrayCumulativeSum(const TArray<float>& InData, TArray<float>& OutData);
+	SIGNALPROCESSING_API void ArrayCumulativeSum(TArrayView<const float> InView, TArray<float>& OutData);
 
 	/** Mean filter of array.
 	 *
 	 *  Note: Uses standard biased mean estimator of Sum(x) / Count(x).
 	 *  Note: At array boundaries, this algorithm truncates windows where no valid array data exists. Values calculated with truncated windows have corresponding increased variances. 
 	 *
-	 *  InData contains data to be filtered.
-	 *  WindowSize determines the number of samples from InData analyzed to produce a value in OutData.
-	 *  WindowOrigin describes the offset from the windows first sample to the index of OutData. For example, if WindowOrigin = WindowSize/4, then OutData[i] = Mean(InData[i - Window/4 : i + 3 * Window / 4]).
+	 *  InView contains data to be filtered.
+	 *  WindowSize determines the number of samples from InView analyzed to produce a value in OutData.
+	 *  WindowOrigin describes the offset from the windows first sample to the index of OutData. For example, if WindowOrigin = WindowSize/4, then OutData[i] = Mean(InView[i - Window/4 : i + 3 * Window / 4]).
 	 *  OutData contains the produceds data.
 	 */
-	SIGNALPROCESSING_API void ArrayMeanFilter(const TArray<float>& InData, int32 WindowSize, int32 WindowOrigin, TArray<float>& OutData);
+	SIGNALPROCESSING_API void ArrayMeanFilter(TArrayView<const float> InView, int32 WindowSize, int32 WindowOrigin, TArray<float>& OutData);
 
 	/** Max filter of array.
 	 *
 	 *  Note: At array boundaries, this algorithm truncates windows where no valid array data exists. 
 	 *
-	 *  InData contains data to be filtered.
-	 *  WindowSize determines the number of samples from InData analyzed to produce a value in OutData.
-	 *  WindowOrigin describes the offset from the windows first sample to the index of OutData. For example, if WindowOrigin = WindowSize/4, then OutData[i] = Max(InData[i - Window/4 : i + 3 * Window / 4]).
+	 *  InView contains data to be filtered.
+	 *  WindowSize determines the number of samples from InView analyzed to produce a value in OutData.
+	 *  WindowOrigin describes the offset from the windows first sample to the index of OutData. For example, if WindowOrigin = WindowSize/4, then OutData[i] = Max(InView[i - Window/4 : i + 3 * Window / 4]).
 	 *  OutData contains the produceds data.
 	 */
-	SIGNALPROCESSING_API void ArrayMaxFilter(const TArray<float>& InData, int32 WindowSize, int32 WindowOrigin, TArray<float>& OutData);
+	SIGNALPROCESSING_API void ArrayMaxFilter(TArrayView<const float> InView, int32 WindowSize, int32 WindowOrigin, TArray<float>& OutData);
 
-	/** Computes the EuclideanNorm of the InArray. Same as calculating the energy in window. */
-	SIGNALPROCESSING_API void ArrayGetEuclideanNorm(const TArray<float>& InArray, float& OutEuclideanNorm);
+	/** Computes the EuclideanNorm of the InView. Same as calculating the energy in window. */
+	SIGNALPROCESSING_API void ArrayGetEuclideanNorm(TArrayView<const float> InView, float& OutEuclideanNorm);
 
-	/** Multiplies each element in InArray by InMultiplier */
-	SIGNALPROCESSING_API void ArrayMultiplyByConstantInPlace(TArray<float>& InArray, float InMultiplier);
+	/** Clamp values in an array.
+	 *
+	 *  InView is a view of a float array to be clamped.
+	 *  InMin is the minimum value allowable in the array.
+	 *  InMax is the maximum value allowable in the array.
+	 */
+	SIGNALPROCESSING_API void ArrayClampInPlace(TArrayView<float> InView, float InMin, float InMax);
+
+	/** Multiplies each element in InView by InMultiplier */
+	SIGNALPROCESSING_API void ArrayMultiplyByConstantInPlace(TArrayView<float> InView, float InMultiplier);
 
-	/** Subtract value from each element in InArray */
-	SIGNALPROCESSING_API void ArraySubtractByConstantInPlace(TArray<float>& InArray, float InSubtrahend);
+	/** Subtract value from each element in InView */
+	SIGNALPROCESSING_API void ArraySubtractByConstantInPlace(TArrayView<float> InView, float InSubtrahend);
 
 	/** Convert magnitude values to decibel values in place. db = 20 * log10(val) */
-	SIGNALPROCESSING_API void ArrayMagnitudeToDecibelInPlace(TArray<float>& InArray);
+	SIGNALPROCESSING_API void ArrayMagnitudeToDecibelInPlace(TArrayView<float> InView);
 
 	/** Convert power values to decibel values in place. db = 10 * log10(val) */
-	SIGNALPROCESSING_API void ArrayPowerToDecibelInPlace(TArray<float>& InArray);
+	SIGNALPROCESSING_API void ArrayPowerToDecibelInPlace(TArrayView<float> InView);
 
 	/** FContiguousSparse2DKernelTransform
 	 *
 	 *  FContiguousSparse2DKernelTransform applies a matrix transformation to an input array. 
-	 *  [OutArray] = [[Kernal]][InArray]  
+	 *  [OutArray] = [[Kernal]][InView]  
 	 *
 	 *  It provides some optimization by exploit the contiguous and sparse qualities of the kernel rows,
 	 *  which allows it to skip multiplications with the number zero. 
@@ -95,17 +103,17 @@ namespace Audio
 
 		/** Transforms the input array given the kernel.
 		 *
-		 *  InArray is the array to be transformed. It must have `NumInElements` number of elements.
+		 *  InView is the array to be transformed. It must have `NumInElements` number of elements.
 		 *  OutArray is the transformed array. It will have `NumOutElements` number of elements.
 		 */
-		void TransformArray(TArrayView<const float> InArray, TArray<float>& OutArray) const;
+		void TransformArray(TArrayView<const float> InView, TArray<float>& OutArray) const;
 
 		/** Transforms the input array given the kernel.
 		 *
-		 *  InArray is the array to be transformed. It must have `NumInElements` number of elements.
+		 *  InView is the array to be transformed. It must have `NumInElements` number of elements.
 		 *  OutArray is the transformed array. It will have `NumOutElements` number of elements.
 		 */
-		void TransformArray(TArrayView<const float> InArray, AlignedFloatBuffer& OutArray) const;
+		void TransformArray(TArrayView<const float> InView, AlignedFloatBuffer& OutArray) const;
 
 		/** Transforms the input array given the kernel.
 		 *
