Adding ConstantQNRT to AudioSynesthesia

diff --git a/Engine/Plugins/Runtime/AudioSynesthesia/Source/AudioSynesthesia/Classes/ConstantQNRT.h b/Engine/Plugins/Runtime/AudioSynesthesia/Source/AudioSynesthesia/Classes/ConstantQNRT.h
new file mode 100644
--- /dev/null
+++ b/Engine/Plugins/Runtime/AudioSynesthesia/Source/AudioSynesthesia/Classes/ConstantQNRT.h
@@ -0,0 +1,157 @@
+// Copyright 1998-2019 Epic Games, Inc. All Rights Reserved.
+
+#pragma once
+
+#include "AudioSynesthesiaNRT.h"
+#include "AudioMixerBlueprintLibrary.h"
+#include "ConstantQNRT.generated.h"
+
+/** EConstantQNormalizationEnum
+ *
+ * Enumeration of available normalization schemes for ConstantQ frequency domain windows.
+ */
+UENUM(BlueprintType)
+enum class EConstantQNormalizationEnum : uint8
+{
+	/** Normalize bands by euclidean norm. Good when using magnitude spectrum. */
+	EqualEuclideanNorm	UMETA(DisplayName="Equal Euclidean Norm"),
+
+	/** Normalize bands by energy. Good when using power spectrum. */
+	EqualEnergy			UMETA(DisplayName="Equal Energy"),
+
+	/** Normalize bands by their maximum values. Will result in relatively strong high frequences because the upper constant Q bands have larger bandwidths. */
+	EqualAmplitude 		UMETA(DisplayName="Equal Amplitude"),
+};
+
+/** EContantQFFTSizeEnum
+ *
+ *  Enumeration of available FFT sizes in audio frames.
+ */
+UENUM(BlueprintType)
+enum class EConstantQFFTSizeEnum : uint8
+{
+	// 64
+	Min,
+	
+	// 128
+	XXSmall,
+
+	// 256
+	XSmall,
+
+	// 512
+	Small,
+
+	// 1024
+	Medium,
+
+	// 2048
+	Large,
+
+	// 4096
+	XLarge,
+
+	// 8192
+	XXLarge,
+
+	// 16384
+	Max
+};
+										 
+/** UConstantQNRTSettings
+ *
+ * Settings for a UConstantQNRT analyzer.
+ */
+UCLASS(Blueprintable)
+class AUDIOSYNESTHESIA_API UConstantQNRTSettings : public UAudioSynesthesiaNRTSettings
+{
+	GENERATED_BODY()
+	public:
+
+		UConstantQNRTSettings();
+
+		/** Starting frequency for first bin of CQT */
+		UPROPERTY(EditAnywhere, BlueprintReadWrite, Category=AudioAnalyzer, meta = (ClampMin = "20.0", ClampMax = "20000"))
+		float StartingFrequency;
+
+        /** Total number of resulting constant Q bands. */
+		UPROPERTY(EditAnywhere, BlueprintReadWrite, Category=AudioAnalyzer, meta = (ClampMin = "1", ClampMax = "96"))
+		int32 NumBands;
+
+		/** Number of bands within an octave. */
+		UPROPERTY(EditAnywhere, BlueprintReadWrite, Category=AudioAnalyzer, meta = (ClampMin = "1", ClampMax = "24"))
+		float NumBandsPerOctave;                        
+
+		/** Number of seconds between cqt measurements */
+		UPROPERTY(EditAnywhere, BlueprintReadWrite, Category=AudioAnalyzer, meta = (ClampMin = "0.01", ClampMax = "1.0"))
+		float AnalysisPeriod;
+
+		/** If true, multichannel audio is downmixed to mono with equal amplitude scaling. If false, each channel gets it's own CQT result. */
+		UPROPERTY(EditAnywhere, BlueprintReadWrite, Category=AudioAnalyzer)
+		bool bDownmixToMono;
+
+		/** Size of FFT. */
+		UPROPERTY(EditAnywhere, BlueprintReadWrite, AdvancedDisplay, Category=AudioAnalyzer)
+		EConstantQFFTSizeEnum FFTSize;
+
+		/** Type of window to be applied to input audio */
+		UPROPERTY(EditAnywhere, BlueprintReadWrite, AdvancedDisplay, Category=AudioAnalyzer)
+        EFFTWindowType WindowType;
+
+		/** Type of spectrum to use. */
+		UPROPERTY(EditAnywhere, BlueprintReadWrite, AdvancedDisplay, Category=AudioAnalyzer)
+		EAudioSpectrumType SpectrumType;
+		
+		/** Stretching factor to control overlap of adjacent bands. */
+		UPROPERTY(EditAnywhere, BlueprintReadWrite, AdvancedDisplay, Category=AudioAnalyzer, meta = (ClampMin = "0.01", ClampMax = "2.0"))
+		float BandWidthStretch;                         
+		
+		/** Normalization scheme used to generate band windows. */
+		UPROPERTY(EditAnywhere, BlueprintReadWrite, AdvancedDisplay, Category=AudioAnalyzer)
+		EConstantQNormalizationEnum CQTNormalization;
+
+		/** Noise floor to use when normalizing CQT */
+		UPROPERTY(EditAnywhere, BlueprintReadWrite, AdvancedDisplay, Category = AudioAnalyzer, meta = (ClampMin = "-120.0", ClampMax="0.0"))
+		float NoiseFloorDb;
+
+		/** Convert UConstantQNRTSettings to FConstantQNRTSettings */
+		TUniquePtr<Audio::IAnalyzerNRTSettings> GetSettings();
+};
+
+
+/** UConstantQNRT
+ *
+ * UConstantQNRT calculates the temporal evolution of constant q transform for a given 
+ * sound. ConstantQ is available for individual channels or the overall sound asset.
+ */
+UCLASS(Blueprintable)
+class AUDIOSYNESTHESIA_API UConstantQNRT : public UAudioSynesthesiaNRT
+{
+	GENERATED_BODY()
+
+	public:
+
+		UConstantQNRT();
+
+		/** The settings for the audio analyzer.  */
+		UPROPERTY(EditAnywhere, BlueprintReadWrite, Category=AudioAnalyzer)
+		UConstantQNRTSettings* Settings;
+
+		/** Get a specific channel cqt of the analyzed sound at a given time. */
+		UFUNCTION(BlueprintCallable, Category="Audio Analyzer")
+		void GetChannelConstantQAtTime(const float InSeconds, const int32 InChannel, TArray<float>& OutConstantQ) const;
+
+		/** Get a specific channel cqt of the analyzed sound at a given time. */
+		UFUNCTION(BlueprintCallable, Category="Audio Analyzer")
+		void GetNormalizedChannelConstantQAtTime(const float InSeconds, const int32 InChannel, TArray<float>& OutConstantQ) const;
+
+		/** Convert ULoudnessNRTSettings to FLoudnessNRTSettings */
+ 		TUniquePtr<Audio::IAnalyzerNRTSettings> GetSettings() const override;
+
+	protected:
+
+		/** Return the name of the IAudioAnalyzerNRTFactory associated with this UAudioAnalyzerNRT */
+		FName GetAnalyzerNRTFactoryName() const override;
+
+};
+
diff --git a/Engine/Plugins/Runtime/AudioSynesthesia/Source/AudioSynesthesia/Private/ConstantQAnalyzer.cpp b/Engine/Plugins/Runtime/AudioSynesthesia/Source/AudioSynesthesia/Private/ConstantQAnalyzer.cpp
new file mode 100644
--- /dev/null
+++ b/Engine/Plugins/Runtime/AudioSynesthesia/Source/AudioSynesthesia/Private/ConstantQAnalyzer.cpp
@@ -0,0 +1,82 @@
+// Copyright 1998-2019 Epic Games, Inc. All Rights Reserved.
+
+#include "ConstantQAnalyzer.h"
+
+#include "DSP/ConstantQ.h"
+#include "DSP/AudioFFT.h"
+#include "DSP/FloatArrayMath.h"
+#include "DSP/BufferVectorOperations.h"
+
+namespace Audio
+{
+	FConstantQAnalyzer::FConstantQAnalyzer(const FConstantQAnalyzerSettings& InSettings, const float InSampleRate)
+	:	Settings(InSettings)
+	,	SampleRate(InSampleRate)
+	,	NumUsefulFFTBins((Settings.FFTSize / 2) + 1)
+    ,   Window(Settings.WindowType, Settings.FFTSize, 1, false)
+	{
+		// Need FFTSize atleast 8 to support optimized operations.
+		check(Settings.FFTSize >= 8);
+		check(Settings.SampleRate > 0.f);
+
+		CQTTransform = NewPseudoConstantQKernelTransform(Settings, Settings.FFTSize, SampleRate);
+
+		// Resize internal buffers
+		WindowedSamples.Reset(Settings.FFTSize);
+		WindowedSamples.AddUninitialized(Settings.FFTSize);
+
+		FFTOutputRealData.Reset(Settings.FFTSize);
+		FFTOutputRealData.AddUninitialized(Settings.FFTSize);
+
+		FFTOutputImagData.Reset(Settings.FFTSize);
+		FFTOutputImagData.AddUninitialized(Settings.FFTSize);
+
+		SpectrumBuffer.Reset(NumUsefulFFTBins);
+		SpectrumBuffer.AddUninitialized(NumUsefulFFTBins);
+	}
+
+	void FConstantQAnalyzer::CalculateCQT(const float* InSamples, TArray<float>& OutCQT)
+	{
+		// Copy input samples and apply window
+		FMemory::Memcpy(WindowedSamples.GetData(), InSamples, sizeof(float) * Settings.FFTSize);
+		Window.ApplyToBuffer(WindowedSamples.GetData());
+
+		// Take FFT of windowed samples
+		const FFTTimeDomainData TimeData = {WindowedSamples.GetData(), Settings.FFTSize};
+		FFTFreqDomainData FreqData = {FFTOutputRealData.GetData(), FFTOutputImagData.GetData()};
+
+		PerformFFT(TimeData, FreqData);
+
+		// Calculate spectrum from FFT output
+		ComputeSpectrum(Settings.SpectrumType, FreqData, Settings.FFTSize, SpectrumBuffer);
+		
+		// Convert spectrum to CQT
+		CQTTransform->TransformArray(SpectrumBuffer, OutCQT);
+
+		// Apply decibel scaling if appropriate.
+		if (EConstantQScaling::Decibel == Settings.Scaling)
+		{
+			// Convert to dB
+			switch (Settings.SpectrumType)
+			{
+				case ESpectrumType::MagnitudeSpectrum:
+					ArrayMagnitudeToDecibelInPlace(OutCQT);
+					break;
+
+				case ESpectrumType::PowerSpectrum:
+					ArrayPowerToDecibelInPlace(OutCQT);
+					break;
+
+				default:
+					checkf(false, TEXT("Unhandled ESpectrumType %s"), GETENUMSTRING(ESpectrumType, Settings.SpectrumType));
+					ArrayPowerToDecibelInPlace(OutCQT);
+			}
+		}
+	}
+
+	const FConstantQAnalyzerSettings& FConstantQAnalyzer::GetSettings() const
+	{
+		return Settings;
+	}
+}
+
diff --git a/Engine/Plugins/Runtime/AudioSynesthesia/Source/AudioSynesthesia/Private/ConstantQNRT.cpp b/Engine/Plugins/Runtime/AudioSynesthesia/Source/AudioSynesthesia/Private/ConstantQNRT.cpp
new file mode 100644
--- /dev/null
+++ b/Engine/Plugins/Runtime/AudioSynesthesia/Source/AudioSynesthesia/Private/ConstantQNRT.cpp
@@ -0,0 +1,239 @@
+// Copyright 1998-2019 Epic Games, Inc. All Rights Reserved.
+
+#include "ConstantQNRT.h"
+#include "ConstantQNRTFactory.h"
+#include "InterpolateSorted.h"
+#include "DSP/ConstantQ.h"
+
+namespace
+{
+	/** Convenience template function for converting from blueprint types to audio types */
+	template<typename InType, typename OutType>
+	OutType ConvertType(const TMap<InType, OutType>& InMap, const InType& InValue, const OutType& InDefaultValue)
+	{
+		checkf(InMap.Contains(InValue), TEXT("Unhandled value conversion"));
+
+		if (InMap.Contains(InValue))
+		{
+			return InMap[InValue];
+		}
+
+		return InDefaultValue;
+	}
+	
+
+	// EConstantQFFTSizeEnum to int
+	int32 BPConstantQFFTSizeEnumToInt(EConstantQFFTSizeEnum InFFTSizeEnum)
+	{
+		static const TMap<EConstantQFFTSizeEnum, int32> FFTSizeMap = {
+			{EConstantQFFTSizeEnum::Min, 		64},
+			{EConstantQFFTSizeEnum::XXSmall, 	128},
+			{EConstantQFFTSizeEnum::XSmall, 	256},
+			{EConstantQFFTSizeEnum::Small,		512},
+			{EConstantQFFTSizeEnum::Medium, 	1024},
+			{EConstantQFFTSizeEnum::Large,		2048},
+			{EConstantQFFTSizeEnum::XLarge, 	4096},
+			{EConstantQFFTSizeEnum::XXLarge, 	8192},
+			{EConstantQFFTSizeEnum::Max, 		16384}
+		};
+
+		return ConvertType(FFTSizeMap, InFFTSizeEnum, 2048);
+
+	}
+
+	// EFFTWindwoType to Audio::EWindowType
+	Audio::EWindowType BPWindowTypeToAudioWindowType(EFFTWindowType InWindowType)
+	{
+		static const TMap<EFFTWindowType, Audio::EWindowType> WindowTypeMap = {
+			{EFFTWindowType::None, 		Audio::EWindowType::None},
+			{EFFTWindowType::Hamming, 	Audio::EWindowType::Hamming},
+			{EFFTWindowType::Hann, 		Audio::EWindowType::Hann},
+			{EFFTWindowType::Blackman,	Audio::EWindowType::Blackman}
+		};
+
+		return ConvertType(WindowTypeMap, InWindowType, Audio::EWindowType::None);
+	}
+
+	// ESpectrumType to Audio::ESpectrumType
+	Audio::ESpectrumType BPSpectrumTypeToAudioSpectrumType(EAudioSpectrumType InSpectrumType)
+	{
+		static const TMap<EAudioSpectrumType, Audio::ESpectrumType> SpectrumTypeMap = {
+			{EAudioSpectrumType::MagnitudeSpectrum,	Audio::ESpectrumType::MagnitudeSpectrum},
+			{EAudioSpectrumType::PowerSpectrum, 	Audio::ESpectrumType::PowerSpectrum}
+		};
+
+		return ConvertType(SpectrumTypeMap, InSpectrumType, Audio::ESpectrumType::PowerSpectrum);
+	}
+
+	// EContantQNormalizationEnum to Audio::EPseudoConstantQNormalization
+	Audio::EPseudoConstantQNormalization BPCQTNormalizationToAudioCQTNormalization(EConstantQNormalizationEnum InNormalization)
+	{
+		static const TMap<EConstantQNormalizationEnum, Audio::EPseudoConstantQNormalization> NormalizationMap = {
+			{EConstantQNormalizationEnum::EqualAmplitude, 		Audio::EPseudoConstantQNormalization::EqualAmplitude},
+			{EConstantQNormalizationEnum::EqualEuclideanNorm,	Audio::EPseudoConstantQNormalization::EqualEuclideanNorm},
+			{EConstantQNormalizationEnum::EqualEnergy, 			Audio::EPseudoConstantQNormalization::EqualEnergy}
+		};
+
+		return ConvertType(NormalizationMap, InNormalization, Audio::EPseudoConstantQNormalization::EqualEnergy);
+	}
+
+	// Interpolates each value in input frame at timestamp
+	void InterpolateSpectrumAtTimestamp(TArrayView<const Audio::FConstantQFrame> InFrames, float InTimestamp, TArray<float>& OutSpectrum)
+	{
+		OutSpectrum.Reset();
+
+		int32 LowerIndex = INDEX_NONE;
+		int32 UpperIndex = INDEX_NONE;
+		float Alpha = 0.f;
+
+		// get interpolation values based on timestmaps
+		GetInterpolationParametersAtTimestamp(InFrames, InTimestamp, LowerIndex, UpperIndex, Alpha);
+
+		if ((INDEX_NONE == LowerIndex) || (INDEX_NONE == UpperIndex))
+		{
+			// there was an error getting the interpolation parameters
+			return;
+		}
+		
+		const Audio::FConstantQFrame& LowerFrame = InFrames[LowerIndex];
+		const Audio::FConstantQFrame& UpperFrame = InFrames[UpperIndex];
+
+		const int32 Num = LowerFrame.Spectrum.Num();
+
+		OutSpectrum.AddUninitialized(Num);
+
+		// Linearly interpolate between upper and lower frame spectrums.
+		const float* LowerSpectrumData = LowerFrame.Spectrum.GetData();
+		const float* UpperSpectrumData = UpperFrame.Spectrum.GetData();
+		float* OutSpectrumData = OutSpectrum.GetData();
+
+		for (int32 i = 0; i < Num; i++)
+		{
+			OutSpectrumData[i] = FMath::Lerp(LowerSpectrumData[i], UpperSpectrumData[i], Alpha);
+		}
+	}
+}
+
+/***************************************************************************/
+/*********************    UConstantQNRTSettings     ************************/
+/***************************************************************************/
+
+UConstantQNRTSettings::UConstantQNRTSettings()
+: 	StartingFrequency(40.f)
+,	NumBands(48)
+,	NumBandsPerOctave(12.f)
+,	AnalysisPeriod(0.01f)
+,	bDownmixToMono(false)
+,	FFTSize(EConstantQFFTSizeEnum::XLarge)
+,	WindowType(EFFTWindowType::Blackman)
+,	SpectrumType(EAudioSpectrumType::PowerSpectrum)
+,	BandWidthStretch(1.f)
+,	CQTNormalization(EConstantQNormalizationEnum::EqualEnergy)
+,	NoiseFloorDb(-60.f)
+{}
+
+/** Convert UConstantQNRTSettings to FConstantQNRTSettings */
+TUniquePtr<Audio::IAnalyzerNRTSettings> UConstantQNRTSettings::GetSettings()
+{
+	TUniquePtr<Audio::FConstantQNRTSettings> Settings = MakeUnique<Audio::FConstantQNRTSettings>();
+
+	Settings->AnalysisPeriod 			= AnalysisPeriod;
+	Settings->bDownmixToMono 			= bDownmixToMono;
+	Settings->FFTSize 					= BPConstantQFFTSizeEnumToInt(FFTSize);
+	Settings->WindowType 				= BPWindowTypeToAudioWindowType(WindowType);
+	Settings->SpectrumType 				= BPSpectrumTypeToAudioSpectrumType(SpectrumType);
+	Settings->NumBands 					= NumBands;
+	Settings->NumBandsPerOctave 		= NumBandsPerOctave;
+	Settings->KernelLowestCenterFreq	= StartingFrequency;
+	Settings->BandWidthStretch 			= BandWidthStretch;
+	Settings->Normalization 			= BPCQTNormalizationToAudioCQTNormalization(CQTNormalization);
+	Settings->Scaling					= Audio::EConstantQScaling::Decibel;
+
+	return Settings;
+}
+
+
+/***************************************************************************/
+/*********************        UConstantQNRT         ************************/
+/***************************************************************************/
+
+UConstantQNRT::UConstantQNRT()
+{
+	Settings = CreateDefaultSubobject<UConstantQNRTSettings>(TEXT("DefaultConstantQNRTSettings"));
+
+#if WITH_EDITOR
+	// Bind settings to audio analyze so changes to default settings will trigger analysis.
+	Settings->AnalyzeAudioDelegate.BindUObject(this, &UAudioAnalyzerNRT::AnalyzeAudio);
+#endif
+}
+
+
+void UConstantQNRT::GetChannelConstantQAtTime(const float InSeconds, const int32 InChannel, TArray<float>& OutConstantQ) const
+{
+	OutConstantQ.Reset();
+
+	TSharedPtr<Audio::FConstantQNRTResult, ESPMode::ThreadSafe> ConstantQResult = GetResult<Audio::FConstantQNRTResult>();
+
+	if (ConstantQResult.IsValid())
+	{
+		if (!ConstantQResult->IsSortedChronologically())
+		{
+			ConstantQResult->SortChronologically();
+		}
+
+		const TArray<Audio::FConstantQFrame>& ConstantQArray = ConstantQResult->GetFramesForChannel(InChannel);
+
+		InterpolateSpectrumAtTimestamp(ConstantQArray, InSeconds, OutConstantQ);
+	}
+}
+
+void UConstantQNRT::GetNormalizedChannelConstantQAtTime(const float InSeconds, const int32 InChannelIdx, TArray<float>& OutConstantQ) const
+{
+	OutConstantQ.Reset();
+
+	TSharedPtr<Audio::FConstantQNRTResult, ESPMode::ThreadSafe> ConstantQResult = GetResult<Audio::FConstantQNRTResult>();
+
+	if (ConstantQResult.IsValid())
+	{
+		if (!ConstantQResult->IsSortedChronologically())
+		{
+			ConstantQResult->SortChronologically();
+		}
+
+		const TArray<Audio::FConstantQFrame>& ConstantQArray = ConstantQResult->GetFramesForChannel(InChannelIdx);
+
+		InterpolateSpectrumAtTimestamp(ConstantQArray, InSeconds, OutConstantQ);
+
+		// Normalize output by subtracting minimum and dividing by range
+		Audio::ArraySubtractByConstantInPlace(OutConstantQ, Settings->NoiseFloorDb);
+
+		float ConstantQRange = ConstantQResult->GetChannelConstantQRange(InChannelIdx, Settings->NoiseFloorDb);
+
+		if (ConstantQRange > SMALL_NUMBER)
+		{
+			const float Scaling = 1.f / ConstantQRange;
+			Audio::ArrayMultiplyByConstantInPlace(OutConstantQ, Scaling);
+		}
+
+		Audio::ArrayClampInPlace(OutConstantQ, 0.f, 1.f);
+	}
+}
+
+TUniquePtr<Audio::IAnalyzerNRTSettings> UConstantQNRT::GetSettings() const
+{
+	TUniquePtr<Audio::IAnalyzerNRTSettings> AnalyzerSettings;
+
+	if (Settings)
+	{
+		AnalyzerSettings = Settings->GetSettings();	
+	}
+
+	return AnalyzerSettings;
+}
+
+FName UConstantQNRT::GetAnalyzerNRTFactoryName() const 
+{
+	static const FName FactoryName(TEXT("ConstantQNRTFactory"));
+	return FactoryName;
+}
+
diff --git a/Engine/Plugins/Runtime/AudioSynesthesia/Source/AudioSynesthesia/Private/ConstantQNRTFactory.cpp b/Engine/Plugins/Runtime/AudioSynesthesia/Source/AudioSynesthesia/Private/ConstantQNRTFactory.cpp
new file mode 100644
--- /dev/null
+++ b/Engine/Plugins/Runtime/AudioSynesthesia/Source/AudioSynesthesia/Private/ConstantQNRTFactory.cpp
@@ -0,0 +1,248 @@
+// Copyright 1998-2019 Epic Games, Inc. All Rights Reserved.
+
+#include "ConstantQNRTFactory.h"
+#include "DSP/ConstantQ.h"
+#include "DSP/DeinterleaveView.h"
+#include "DSP/FloatArrayMath.h"
+
+namespace Audio 
+{
+	/*******************************************************************************/
+	/*************************** FConstantQNRTResult *******************************/
+	/*******************************************************************************/
+
+	FArchive &operator <<(FArchive& Ar, FConstantQFrame& Frame)
+	{
+		Ar << Frame.Channel;
+		Ar << Frame.Timestamp;
+		Ar << Frame.Spectrum;
+
+		return Ar;
+	}
+
+	FConstantQNRTResult::FConstantQNRTResult()
+	:	bIsSortedChronologically(false)
+	{}
+
+	void FConstantQNRTResult::Serialize(FArchive& Archive)
+	{
+		Archive << bIsSortedChronologically;
+		Archive << ChannelCQTFrames;
+		Archive << ChannelCQTIntervals;
+	}
+
+	void FConstantQNRTResult::AddFrame(int32 InChannelIndex, float InTimestamp, TArrayView<const float> InSpectrum)
+	{
+		TArray<FConstantQFrame>& FrameArray = ChannelCQTFrames.FindOrAdd(InChannelIndex);
+		FrameArray.Emplace(InChannelIndex, InTimestamp, InSpectrum);
+
+		FFloatInterval& Interval = ChannelCQTIntervals.FindOrAdd(InChannelIndex);
+		for (float Value : InSpectrum)
+		{
+			Interval.Include(Value);
+		}
+
+		bIsSortedChronologically = false;
+	}
+
+	const TArray<FConstantQFrame>& FConstantQNRTResult::GetFramesForChannel(int32 InChannelIndex) const
+	{
+		check(InChannelIndex < ChannelCQTFrames.Num());
+		return ChannelCQTFrames[InChannelIndex];
+	}
+
+	float FConstantQNRTResult::GetChannelConstantQRange(int32 InChannelIdx, float InNoiseFloor) const
+	{
+		const FFloatInterval& Interval = ChannelCQTIntervals.FindRef(InChannelIdx);
+
+		if (Interval.Contains(InNoiseFloor))
+		{
+			return Interval.Max - InNoiseFloor;
+		}
+		else if (InNoiseFloor > Interval.Max)
+		{
+			return 0.f;
+		}
+
+		return Interval.Size();
+	}
+
+	int32 FConstantQNRTResult::GetNumChannels() const
+	{
+		return ChannelCQTFrames.Num();
+	}
+
+	 /* Returns true if FConstantQNRTFrame arrays are sorted in chronologically ascending order via their timestamp.
+	 */
+	bool FConstantQNRTResult::IsSortedChronologically() const
+	{
+		return bIsSortedChronologically;
+	}
+
+	/**
+	 * Sorts FConstantQNRTFrame arrays in chronologically ascnding order via their timestamp.
+	 */
+	void FConstantQNRTResult::SortChronologically()
+	{
+		for (auto& KeyValue : ChannelCQTFrames)
+		{
+			KeyValue.Value.Sort([](const FConstantQFrame& A, const FConstantQFrame& B) { return A.Timestamp < B.Timestamp; });
+		}
+
+		bIsSortedChronologically = true;
+	}
+
+	/*******************************************************************************/
+	/*************************** FConstantQNRTWorker *******************************/
+	/*******************************************************************************/
+
+	FConstantQNRTWorker::FConstantQNRTWorker(const FAnalyzerNRTParameters& InParams, const FConstantQNRTSettings& InAnalyzerSettings)
+	:	NumChannels(InParams.NumChannels)
+	,	NumBuffers(0)
+	,	SampleRate(InParams.SampleRate)
+	,	NumHopFrames(0)
+	,	NumHopSamples(0)
+	,	NumWindowFrames(InAnalyzerSettings.FFTSize)
+	,	NumWindowSamples(InAnalyzerSettings.FFTSize * InParams.NumChannels)
+	,	MonoScaling(1.f)
+	,	bDownmixToMono(InAnalyzerSettings.bDownmixToMono)
+	{
+		check(InParams.NumChannels > 0);
+		check(InParams.SampleRate > 0.f);
+
+		check(InAnalyzerSettings.AnalysisPeriod > 0.f);
+		check(FMath::IsPowerOfTwo(InAnalyzerSettings.FFTSize));
+
+
+		NumHopFrames = FMath::Max(1, FMath::RoundToInt(InAnalyzerSettings.AnalysisPeriod * InParams.SampleRate));
+		NumHopSamples = NumHopFrames * NumChannels;
+
+		if (bDownmixToMono)
+		{
+			MonoBuffer.Reset(NumWindowFrames);
+			MonoBuffer.AddUninitialized(NumWindowFrames);
+			// equal power mixing
+			MonoScaling = 1.f / FMath::Sqrt(static_cast<float>(NumChannels));
+		}
+
+		SlidingBuffer = MakeUnique<TSlidingBuffer<float> >(NumWindowSamples, NumHopSamples);
+
+		ConstantQAnalyzer = MakeUnique<FConstantQAnalyzer>(InAnalyzerSettings, SampleRate);
+	}
+
+
+	void FConstantQNRTWorker::Analyze(TArrayView<const float> InAudio, IAnalyzerNRTResult* OutResult)
+	{
+		bool bDoFlush = false;
+
+		AnalyzeMultichannel(InAudio, OutResult, bDoFlush);
+	}
+
+	void FConstantQNRTWorker::AnalyzeMultichannel(TArrayView<const float> InAudio, IAnalyzerNRTResult* OutResult, bool bDoFlush)
+	{
+		// Assume that outer layers ensured that this is of correct type.
+		FConstantQNRTResult* ConstantQResult = static_cast<FConstantQNRTResult*>(OutResult);
+
+		check(nullptr != ConstantQResult);
+
+		TAutoSlidingWindow<float, FAudioBufferAlignedAllocator> SlidingWindow(*SlidingBuffer, InAudio, HopBuffer, bDoFlush);
+
+		if (bDownmixToMono)
+		{
+			for (const AlignedFloatBuffer& Window : SlidingWindow)
+			{
+				FMemory::Memset(MonoBuffer.GetData(), 0, sizeof(float) * NumWindowFrames);
+
+				// Sum all channels into mono buffer
+				TAutoDeinterleaveView<float, FAudioBufferAlignedAllocator> DeinterleaveView(Window, ChannelBuffer, NumChannels);
+				for(auto Channel : DeinterleaveView)
+				{
+					MixInBufferFast(Channel.Values, MonoBuffer);		
+				}
+
+				// Equal power sum. assuming incoherrent signals.
+				MultiplyBufferByConstantInPlace(MonoBuffer, 1.f / FMath::Sqrt(static_cast<float>(NumChannels)));
+
+				AnalyzeWindow(MonoBuffer, 0, *ConstantQResult);
+
+				NumBuffers++;
+			}
+		}
+		else
+		{
+			// Loop through entire array of input samples.
+			for (const AlignedFloatBuffer& Window : SlidingWindow)
+			{
+				// Each channel is analyzed seperately.
+				for(const TDeinterleaveView<float>::TChannel<FAudioBufferAlignedAllocator>& Channel : TAutoDeinterleaveView<float, FAudioBufferAlignedAllocator>(Window, ChannelBuffer, NumChannels))
+				{
+					AnalyzeWindow(Channel.Values, Channel.ChannelIndex, *ConstantQResult);
+				}
+
+				NumBuffers++;
+			}
+		}
+	}
+
+	/** Called when analysis of audio asset is complete. */
+	void FConstantQNRTWorker::Finalize(IAnalyzerNRTResult* OutResult) 
+	{
+		AlignedFloatBuffer EmptyArray;
+		bool bDoFlush = true;
+
+		AnalyzeMultichannel(EmptyArray, OutResult, bDoFlush);
+
+		// Assume that outer layers ensured that this is of correct type.
+		FConstantQNRTResult* ConstantQResult = static_cast<FConstantQNRTResult*>(OutResult);
+
+		check(nullptr != ConstantQResult);
+
+		NumBuffers = 0; // Reset internal counters
+
+		if (nullptr != OutResult)
+		{
+			ConstantQResult->SortChronologically();
+		}
+	}
+
+	void FConstantQNRTWorker::AnalyzeWindow(const AlignedFloatBuffer& InWindow, int32 InChannelIndex, FConstantQNRTResult& OutResult)
+	{
+		// Run CQT Analyzer
+		ConstantQAnalyzer->CalculateCQT(InWindow.GetData(), CQTSpectrum);
+
+		// Get timestamp as center of audio window.
+		float Timestamp = ((NumBuffers * NumHopFrames) + (0.5f * ConstantQAnalyzer->GetSettings().FFTSize)) / SampleRate;
+
+		OutResult.AddFrame(InChannelIndex, Timestamp, CQTSpectrum);
+	}
+
+	// Name of this analyzer type.
+	FName FConstantQNRTFactory::GetName() const
+	{
+		static FName FactoryName(TEXT("ConstantQNRTFactory"));
+		return FactoryName;
+	}
+
+	// Human readable name of this analyzer.
+	FString FConstantQNRTFactory::GetTitle() const
+	{
+		return TEXT("Constant Q Analyzer Non-Real-Time");
+	}
+
+	// Create a new FConstantQNRTResult.
+	TUniquePtr<IAnalyzerNRTResult> FConstantQNRTFactory::NewResult()
+	{
+		return MakeUnique<FConstantQNRTResult>();
+	}
+
+	// Create a new FConstantQNRTWorker
+	TUniquePtr<IAnalyzerNRTWorker> FConstantQNRTFactory::NewWorker(const FAnalyzerNRTParameters& InParams, const IAnalyzerNRTSettings* InSettings)
+	{
+		const FConstantQNRTSettings* ConstantQSettings = static_cast<const FConstantQNRTSettings*>(InSettings);
+
+		check(nullptr != ConstantQSettings);
+
+		return MakeUnique<FConstantQNRTWorker>(InParams, *ConstantQSettings);
+	}
+}
+
diff --git a/Engine/Plugins/Runtime/AudioSynesthesia/Source/AudioSynesthesia/Public/ConstantQAnalyzer.h b/Engine/Plugins/Runtime/AudioSynesthesia/Source/AudioSynesthesia/Public/ConstantQAnalyzer.h
new file mode 100644
--- /dev/null
+++ b/Engine/Plugins/Runtime/AudioSynesthesia/Source/AudioSynesthesia/Public/ConstantQAnalyzer.h
@@ -0,0 +1,93 @@
+// Copyright 1998-2019 Epic Games, Inc. All Rights Reserved.
+
+#pragma once
+
+#include "DSP/ConstantQ.h"
+#include "DSP/AudioFFT.h"
+
+namespace Audio
+{
+	/** Denotes which magnitude scaling to apply to output */
+	enum class EConstantQScaling : uint8
+	{
+		/** CQT spectrum values conform to ESpectrumType, either magnitude or power. */
+		Linear,
+
+		/** CQT spectrum values are in decibels. */
+		Decibel
+	};
+
+	/** FConstantQAnalyzerSettings
+	 *
+	 * 	Settings for the Constant Q Analyzer.
+	 */
+	struct FConstantQAnalyzerSettings : public FPseudoConstantQKernelSettings
+	{
+		/** Size of FFT describe in number of samples */
+		int32 FFTSize;
+
+		/** Type of window to be applied to input audio */
+        EWindowType WindowType;
+
+		/** Type of spectrum to use. */
+		ESpectrumType SpectrumType;
+
+		/** Type of scaling to use. */
+		EConstantQScaling Scaling;
+
+		FConstantQAnalyzerSettings()
+		:	FFTSize(1024)
+		,	WindowType(EWindowType::Blackman)
+		,	SpectrumType(ESpectrumType::PowerSpectrum)
+		,	Scaling(EConstantQScaling::Decibel)	
+		{}
+	};
+
+	/** FConstantQAnalyzer
+	 *
+	 *  The Constant Q Analyzer computs a pseudo constant q transform of the input audio.
+	 *  ConstantQ transforms produce frequency strength at logarithmically spaced frequency
+	 *  intervals (as opposed to the linearly spaced intervals of the FFT). Logarithmically
+	 *  spaced frequency intervals map nicely to musical scales.
+	 */
+	class AUDIOSYNESTHESIA_API FConstantQAnalyzer
+	{
+	public:
+
+		/** Construct a FConstantQAnalyzer
+		 *
+		 * 	InSettings contains the parameters used to configure the analyzer.
+		 * 	InSampleRate is the sample rate of the input audio.
+		 */
+		FConstantQAnalyzer(const FConstantQAnalyzerSettings& InSettings, const float InSampleRate);
+
+		/** Calcualte Constant Q Transform of input audio.
+		 *
+		 * InSamples is an array of input audio. It must be of length GetSettings().FFTSize.
+		 * OutCQT will be populated with the constants q transform values. It will have GetSettings().NumBands values.
+		 */
+		void CalculateCQT(const float* InSamples, TArray<float>& OutCQT);
+
+		/** Returns a reference to the settings used in this analyzer */
+		const FConstantQAnalyzerSettings& GetSettings() const;
+
+	private:
+
+		const FConstantQAnalyzerSettings Settings;
+
+		const float SampleRate;
+
+		/** Only the first (FFTSize / 2) + 1 FFT bins are used since FFT of real values is conjugate symmetric. */
+		const int32 NumUsefulFFTBins;
+
+		/** Scale factor to normalize between different FFT Sizes */
+		FWindow Window;
+
+		AlignedFloatBuffer WindowedSamples;
+		AlignedFloatBuffer FFTOutputRealData;
+		AlignedFloatBuffer FFTOutputImagData;
+		AlignedFloatBuffer SpectrumBuffer;
+
+		TUniquePtr<FContiguousSparse2DKernelTransform> CQTTransform;
+	};
+}
diff --git a/Engine/Plugins/Runtime/AudioSynesthesia/Source/AudioSynesthesia/Public/ConstantQNRTFactory.h b/Engine/Plugins/Runtime/AudioSynesthesia/Source/AudioSynesthesia/Public/ConstantQNRTFactory.h
new file mode 100644
--- /dev/null
+++ b/Engine/Plugins/Runtime/AudioSynesthesia/Source/AudioSynesthesia/Public/ConstantQNRTFactory.h
@@ -0,0 +1,194 @@
+// Copyright 1998-2019 Epic Games, Inc. All Rights Reserved.
+
+#pragma once
+
+#include "IAudioAnalyzerNRTInterface.h"
+#include "DSP/ConstantQ.h"
+#include "DSP/SlidingWindow.h"
+#include "DSP/FloatArrayMath.h"
+#include "ConstantQAnalyzer.h"
+
+namespace Audio
+{
+	/** FConstantQNRTSettings
+	 *
+	 *  Settings for the Constant Q Non-Real-Time Factory.
+	 */
+	class AUDIOSYNESTHESIA_API FConstantQNRTSettings : public IAnalyzerNRTSettings, public FConstantQAnalyzerSettings 
+	{
+	public:
+
+		/** Time, in seconds, between constant Q frames. */
+		float AnalysisPeriod;
+
+		/** 
+		 * If true, all channels are mixed together with equal power before analysis. Only one channel
+		 * is produced at channel index 0.
+		 * If false, then each channel is analyzed separately. 
+		 */
+		bool bDownmixToMono;
+
+		FConstantQNRTSettings()
+		:	AnalysisPeriod(0.01f)
+		,	bDownmixToMono(false)
+		{}
+	};
+
+
+	/** FConstantQFrame
+	 * 
+	 *  Contains Constant Q data relating to one audio window.
+	 */
+	struct AUDIOSYNESTHESIA_API FConstantQFrame
+	{
+		/** Audio channel which produced the data. */
+		int32 Channel;
+		
+		/** Timestamp in seconds referring to the center of the audio window. */
+		float Timestamp;
+
+		/** Output spectral data */
+		TArray<float> Spectrum;
+
+		FConstantQFrame()
+		:	Channel(0)
+		,	Timestamp(0.f)
+		{}
+
+		FConstantQFrame(int32 InChannelIndex, float InTimestamp, TArrayView<const float> InSpectrum)
+		:	Channel(InChannelIndex)
+		,	Timestamp(InTimestamp)
+		,	Spectrum(InSpectrum.GetData(), InSpectrum.Num())
+		{}
+	};
+
+	/** Serialize FConstantQFrame */
+	AUDIOSYNESTHESIA_API FArchive &operator <<(FArchive& Ar, FConstantQFrame& Frame);
+
+	/** FConstantQNRTResult
+	 *
+	 *  FConstantQNRTResult is a container for the output of the FConstantQNRTWorker.
+	 */
+	class AUDIOSYNESTHESIA_API FConstantQNRTResult : public IAnalyzerNRTResult
+	{
+	public:
+
+		FConstantQNRTResult();
+
+		/** Serialize or unserialize object */
+		void Serialize(FArchive& Archive) override;
+
+		/** Add a single frame of CQT data */
+		void AddFrame(int32 InChannelIndex, float InTimestamp, TArrayView<const float> InSpectrum);
+
+		/** Retrieve the array of frames for a single channel of audio. */
+		const TArray<FConstantQFrame>& GetFramesForChannel(int32 InChannelIndex) const;
+
+		/** Retrieve the difference between the maximum and minimum value in the spectrum. If 
+		 *  the minimum value is less than InNoiseFloor, then the result will be the difference
+		 *  betwen the maximum value and the noise floor. 
+		 */
+		float GetChannelConstantQRange(int32 InChannelIdx, float InNoiseFloor) const;
+
+		/** Retrieve the number of channels available in the result. */
+		int32 GetNumChannels() const;
+
+	 	/** Returns true if FConstantQFrame arrays are sorted in chronologically ascending order via their timestamp.  */
+		bool IsSortedChronologically() const;
+
+		/** Sorts FConstantQFrame arrays in chronologically ascnding order via their timestamp.  */
+		void SortChronologically();
+
+	private:
+
+		TMap<int32, TArray<FConstantQFrame> > ChannelCQTFrames;
+
+		TMap<int32, FFloatInterval> ChannelCQTIntervals;
+
+		bool bIsSortedChronologically;
+	};
+
+	/** FConstantQNRTWorker
+	 *
+	 *  FConstantQNRTWorker computes a FConstantQNRTResult from audio samples.
+	 */
+	class AUDIOSYNESTHESIA_API FConstantQNRTWorker : public IAnalyzerNRTWorker
+	{
+	public:
+
+		/** Constructor
+		 *
+		 * InParams are the parameters which describe characteristics of the input audio.
+		 * InAnalyzerSettings are the settings which control various aspects of the algorithm.
+		 */
+		FConstantQNRTWorker(const FAnalyzerNRTParameters& InParams, const FConstantQNRTSettings& InAnalyzerSettings);
+
+		/**
+		 *  Analyze audio and put results into results pointer.
+		 *
+		 *  InAudio is an array view of audio.
+		 *  OutResult is a pointer to a valid FConstantQNRTResult
+		 */
+		void Analyze(TArrayView<const float> InAudio, IAnalyzerNRTResult* OutResult) override;
+
+		/** 
+		 *  Call when analysis of audio asset is complete. 
+		 *
+		 *  OutResult must be a pointer to a valid FConstantQNRTResult. 
+		 */
+		void Finalize(IAnalyzerNRTResult* OutResult) override;
+
+	private:
+
+		/** Analyze audio with multiple channels interleaved. */
+		void AnalyzeMultichannel(TArrayView<const float> InAudio, IAnalyzerNRTResult* OutResult, bool bDoFlush);
+
+		/** Analyze a single window of audio from a single channel */
+		void AnalyzeWindow(const AlignedFloatBuffer& InWindow, int32 InChannelIndex, FConstantQNRTResult& OutResult);
+
+		int32 NumChannels;
+		int32 NumBuffers;
+		float SampleRate;
+		int32 NumHopFrames;
+		int32 NumHopSamples;
+		int32 NumWindowFrames;
+		int32 NumWindowSamples;
+
+		float MonoScaling;
+
+		AlignedFloatBuffer HopBuffer;
+		AlignedFloatBuffer ChannelBuffer;
+		AlignedFloatBuffer MonoBuffer;
+
+		TArray<float> CQTSpectrum;
+
+		TUniquePtr<TSlidingBuffer<float> > SlidingBuffer;
+		TUniquePtr<FConstantQAnalyzer> ConstantQAnalyzer;
+
+		bool bDownmixToMono;
+	};
+
+	/** FConstantQNRTFactory
+	 *  
+	 *  Factory for creating FConstantQNRT workers and results
+	 */
+	class AUDIOSYNESTHESIA_API FConstantQNRTFactory : public IAnalyzerNRTFactory
+	{
+	public:
+
+		/** Name of this analyzer type. */
+		FName GetName() const override;
+
+		/** Human readable name of this analyzer. */
+		FString GetTitle() const override;
+
+		/** Create a new FConstantQNRTResult. */
+		TUniquePtr<IAnalyzerNRTResult> NewResult() override;
+
+		/** Create a new FConstantQNRTWorker 
+		 *
+		 *  InSettings must be a pointer to FConstantQNRTSetting
+		 */
+		TUniquePtr<IAnalyzerNRTWorker> NewWorker(const FAnalyzerNRTParameters& InParams, const IAnalyzerNRTSettings* InSettings) override;
+	};
+}
diff --git a/Engine/Source/Runtime/AudioMixer/Private/DSP/ConstantQ.cpp b/Engine/Source/Runtime/AudioMixer/Private/DSP/ConstantQ.cpp
--- a/Engine/Source/Runtime/AudioMixer/Private/DSP/ConstantQ.cpp
+++ b/Engine/Source/Runtime/AudioMixer/Private/DSP/ConstantQ.cpp
@@ -61,7 +61,7 @@ namespace Audio
 		check(InSampleRate > 0.f);
 		check(InFFTSize > 0);
 
-		const int32 NumUsefulFFTBins = InFFTSize / 2 + 1;
+		const int32 NumUsefulFFTBins = (InFFTSize / 2) + 1;
 
 		TUniquePtr<FContiguousSparse2DKernelTransform> Transform = MakeUnique<FContiguousSparse2DKernelTransform>(NumUsefulFFTBins, InSettings.NumBands);
 
@@ -73,7 +73,11 @@ namespace Audio
 
 			if ((BandCenter - BandWidth) > InSampleRate)
 			{
-				// TODO: have this check be fore a very large number
+				continue;
+			}
+
+			if (BandCenter > (2.f * InSampleRate))
+			{
 				continue;
 			}
 
