Fix signature to match function being overriden
#rnx

diff --git a/Engine/Plugins/Runtime/AudioSynesthesia/Source/AudioSynesthesia/Classes/ConstantQNRT.h b/Engine/Plugins/Runtime/AudioSynesthesia/Source/AudioSynesthesia/Classes/ConstantQNRT.h
--- a/Engine/Plugins/Runtime/AudioSynesthesia/Source/AudioSynesthesia/Classes/ConstantQNRT.h
+++ b/Engine/Plugins/Runtime/AudioSynesthesia/Source/AudioSynesthesia/Classes/ConstantQNRT.h
@@ -146,12 +146,12 @@ class AUDIOSYNESTHESIA_API UConstantQNRT : public UAudioSynesthesiaNRT
 		void GetNormalizedChannelConstantQAtTime(const float InSeconds, const int32 InChannel, TArray<float>& OutConstantQ) const;
 
 		/** Convert ULoudnessNRTSettings to FLoudnessNRTSettings */
- 		TUniquePtr<Audio::IAnalyzerNRTSettings> GetSettings() const override;
+ 		virtual TUniquePtr<Audio::IAnalyzerNRTSettings> GetSettings() override;
 
 	protected:
 
 		/** Return the name of the IAudioAnalyzerNRTFactory associated with this UAudioAnalyzerNRT */
-		FName GetAnalyzerNRTFactoryName() const override;
+		virtual FName GetAnalyzerNRTFactoryName() const override;
 
 };
 
diff --git a/Engine/Plugins/Runtime/AudioSynesthesia/Source/AudioSynesthesia/Private/ConstantQNRT.cpp b/Engine/Plugins/Runtime/AudioSynesthesia/Source/AudioSynesthesia/Private/ConstantQNRT.cpp
--- a/Engine/Plugins/Runtime/AudioSynesthesia/Source/AudioSynesthesia/Private/ConstantQNRT.cpp
+++ b/Engine/Plugins/Runtime/AudioSynesthesia/Source/AudioSynesthesia/Private/ConstantQNRT.cpp
@@ -219,7 +219,7 @@ void UConstantQNRT::GetNormalizedChannelConstantQAtTime(const float InSeconds, c
 	}
 }
 
-TUniquePtr<Audio::IAnalyzerNRTSettings> UConstantQNRT::GetSettings() const
+TUniquePtr<Audio::IAnalyzerNRTSettings> UConstantQNRT::GetSettings()
 {
 	TUniquePtr<Audio::IAnalyzerNRTSettings> AnalyzerSettings;
 
