Fix compilation errors
#rnx

diff --git a/Engine/Plugins/Runtime/AudioSynesthesia/Source/AudioSynesthesia/Private/ConstantQAnalyzer.cpp b/Engine/Plugins/Runtime/AudioSynesthesia/Source/AudioSynesthesia/Private/ConstantQAnalyzer.cpp
--- a/Engine/Plugins/Runtime/AudioSynesthesia/Source/AudioSynesthesia/Private/ConstantQAnalyzer.cpp
+++ b/Engine/Plugins/Runtime/AudioSynesthesia/Source/AudioSynesthesia/Private/ConstantQAnalyzer.cpp
@@ -17,7 +17,7 @@ namespace Audio
 	{
 		// Need FFTSize atleast 8 to support optimized operations.
 		check(Settings.FFTSize >= 8);
-		check(Settings.SampleRate > 0.f);
+		check(SampleRate > 0.f);
 
 		CQTTransform = NewPseudoConstantQKernelTransform(Settings, Settings.FFTSize, SampleRate);
 
@@ -68,7 +68,8 @@ namespace Audio
 					break;
 
 				default:
-					checkf(false, TEXT("Unhandled ESpectrumType %s"), GETENUMSTRING(ESpectrumType, Settings.SpectrumType));
+					check(false);
+					//checkf(false, TEXT("Unhandled ESpectrumType %s"), GETENUMSTRING(ESpectrumType, Settings.SpectrumType));
 					ArrayPowerToDecibelInPlace(OutCQT);
 			}
 		}
