Remove const as non-const functions are called inside them
#rnx

diff --git a/Engine/Plugins/Runtime/AudioSynesthesia/Source/AudioSynesthesia/Classes/ConstantQNRT.h b/Engine/Plugins/Runtime/AudioSynesthesia/Source/AudioSynesthesia/Classes/ConstantQNRT.h
--- a/Engine/Plugins/Runtime/AudioSynesthesia/Source/AudioSynesthesia/Classes/ConstantQNRT.h
+++ b/Engine/Plugins/Runtime/AudioSynesthesia/Source/AudioSynesthesia/Classes/ConstantQNRT.h
@@ -139,11 +139,11 @@ class AUDIOSYNESTHESIA_API UConstantQNRT : public UAudioSynesthesiaNRT
 
 		/** Get a specific channel cqt of the analyzed sound at a given time. */
 		UFUNCTION(BlueprintCallable, Category="Audio Analyzer")
-		void GetChannelConstantQAtTime(const float InSeconds, const int32 InChannel, TArray<float>& OutConstantQ) const;
+		void GetChannelConstantQAtTime(const float InSeconds, const int32 InChannel, TArray<float>& OutConstantQ);
 
 		/** Get a specific channel cqt of the analyzed sound at a given time. */
 		UFUNCTION(BlueprintCallable, Category="Audio Analyzer")
-		void GetNormalizedChannelConstantQAtTime(const float InSeconds, const int32 InChannel, TArray<float>& OutConstantQ) const;
+		void GetNormalizedChannelConstantQAtTime(const float InSeconds, const int32 InChannel, TArray<float>& OutConstantQ);
 
 		/** Convert ULoudnessNRTSettings to FLoudnessNRTSettings */
  		virtual TUniquePtr<Audio::IAnalyzerNRTSettings> GetSettings() override;
diff --git a/Engine/Plugins/Runtime/AudioSynesthesia/Source/AudioSynesthesia/Private/ConstantQNRT.cpp b/Engine/Plugins/Runtime/AudioSynesthesia/Source/AudioSynesthesia/Private/ConstantQNRT.cpp
--- a/Engine/Plugins/Runtime/AudioSynesthesia/Source/AudioSynesthesia/Private/ConstantQNRT.cpp
+++ b/Engine/Plugins/Runtime/AudioSynesthesia/Source/AudioSynesthesia/Private/ConstantQNRT.cpp
@@ -168,7 +168,7 @@ UConstantQNRT::UConstantQNRT()
 }
 
 
-void UConstantQNRT::GetChannelConstantQAtTime(const float InSeconds, const int32 InChannel, TArray<float>& OutConstantQ) const
+void UConstantQNRT::GetChannelConstantQAtTime(const float InSeconds, const int32 InChannel, TArray<float>& OutConstantQ)
 {
 	OutConstantQ.Reset();
 
@@ -187,7 +187,7 @@ void UConstantQNRT::GetChannelConstantQAtTime(const float InSeconds, const int32
 	}
 }
 
-void UConstantQNRT::GetNormalizedChannelConstantQAtTime(const float InSeconds, const int32 InChannelIdx, TArray<float>& OutConstantQ) const
+void UConstantQNRT::GetNormalizedChannelConstantQAtTime(const float InSeconds, const int32 InChannelIdx, TArray<float>& OutConstantQ)
 {
 	OutConstantQ.Reset();
 
