Adding AudioMixer dependency and fixing include in order to fix build.

diff --git a/Engine/Plugins/Runtime/AudioSynesthesia/Source/AudioSynesthesia/AudioSynesthesia.Build.cs b/Engine/Plugins/Runtime/AudioSynesthesia/Source/AudioSynesthesia/AudioSynesthesia.Build.cs
--- a/Engine/Plugins/Runtime/AudioSynesthesia/Source/AudioSynesthesia/AudioSynesthesia.Build.cs
+++ b/Engine/Plugins/Runtime/AudioSynesthesia/Source/AudioSynesthesia/AudioSynesthesia.Build.cs
@@ -18,6 +18,7 @@ namespace UnrealBuildTool.Rules
 					"Engine",
 					"SignalProcessing",
 					"AudioMixerCore",
+					"AudioMixer",
                     "AudioAnalyzer"
                 }
             );
