Trivial fix to remove incorrect check.

diff --git a/Engine/Plugins/Runtime/AudioSynesthesia/Source/AudioSynesthesia/Private/ConstantQNRTFactory.cpp b/Engine/Plugins/Runtime/AudioSynesthesia/Source/AudioSynesthesia/Private/ConstantQNRTFactory.cpp
--- a/Engine/Plugins/Runtime/AudioSynesthesia/Source/AudioSynesthesia/Private/ConstantQNRTFactory.cpp
+++ b/Engine/Plugins/Runtime/AudioSynesthesia/Source/AudioSynesthesia/Private/ConstantQNRTFactory.cpp
@@ -47,7 +47,6 @@ namespace Audio
 
 	const TArray<FConstantQFrame>& FConstantQNRTResult::GetFramesForChannel(int32 InChannelIndex) const
 	{
-		check(InChannelIndex < ChannelCQTFrames.Num());
 		return ChannelCQTFrames[InChannelIndex];
 	}
 
