Push remaining types to submenus in sound editor asset creation menu
- Minor: Add missing editor API to requiring types in IAudioExtensionPlugin

diff --git a/Engine/Plugins/Runtime/AudioSynesthesia/Source/AudioSynesthesia/Classes/AudioSynesthesiaNRT.h b/Engine/Plugins/Runtime/AudioSynesthesia/Source/AudioSynesthesia/Classes/AudioSynesthesiaNRT.h
--- a/Engine/Plugins/Runtime/AudioSynesthesia/Source/AudioSynesthesia/Classes/AudioSynesthesiaNRT.h
+++ b/Engine/Plugins/Runtime/AudioSynesthesia/Source/AudioSynesthesia/Classes/AudioSynesthesiaNRT.h
@@ -17,6 +17,8 @@ class AUDIOSYNESTHESIA_API UAudioSynesthesiaNRTSettings : public UAudioAnalyzerN
 
 		FText GetAssetActionName() const override;
 
+		const TArray<FText>& GetAssetActionSubmenus() const;
+
 		UClass* GetSupportedClass() const override;
 
 		FColor GetTypeColor() const override;
@@ -35,6 +37,8 @@ class AUDIOSYNESTHESIA_API UAudioSynesthesiaNRT : public UAudioAnalyzerNRT
 
 		FText GetAssetActionName() const override;
 
+		const TArray<FText>& GetAssetActionSubmenus() const;
+
 		UClass* GetSupportedClass() const override;
 
 		FColor GetTypeColor() const override;
diff --git a/Engine/Plugins/Runtime/AudioSynesthesia/Source/AudioSynesthesia/Private/AudioSynesthesiaNRT.cpp b/Engine/Plugins/Runtime/AudioSynesthesia/Source/AudioSynesthesia/Private/AudioSynesthesiaNRT.cpp
--- a/Engine/Plugins/Runtime/AudioSynesthesia/Source/AudioSynesthesia/Private/AudioSynesthesiaNRT.cpp
+++ b/Engine/Plugins/Runtime/AudioSynesthesia/Source/AudioSynesthesia/Private/AudioSynesthesiaNRT.cpp
@@ -5,7 +5,17 @@
 
 FText UAudioSynesthesiaNRTSettings::GetAssetActionName() const
 {
-	return FText::FromString("Audio Synesthesia NRT Settings"); 
+	return NSLOCTEXT("AssetTypeActions", "AssetTypeActions_AssetSoundSynesthesiaNRTSettings", "Synesthesia NRT Settings");
+}
+
+const TArray<FText>& UAudioSynesthesiaNRTSettings::GetAssetActionSubmenus() const
+{
+	static const TArray<FText> SubMenus
+	{
+		NSLOCTEXT("AssetTypeActions", "AssetTypeActions_AssetSoundAnalysisSubmenu", "Analysis")
+	};
+
+	return SubMenus;
 }
 
 UClass* UAudioSynesthesiaNRTSettings::GetSupportedClass() const 
@@ -20,7 +30,17 @@ FColor UAudioSynesthesiaNRTSettings::GetTypeColor() const
 
 FText UAudioSynesthesiaNRT::GetAssetActionName() const
 {
-	return FText::FromString("Audio Synesthesia NRT"); 
+	return NSLOCTEXT("AssetTypeActions", "AssetTypeActions_AssetSoundSynesthesiaNRT", "Synesthesia NRT");
+}
+
+const TArray<FText>& UAudioSynesthesiaNRT::GetAssetActionSubmenus() const
+{
+	static const TArray<FText> SubMenus
+	{
+		NSLOCTEXT("AssetTypeActions", "AssetTypeActions_AssetSoundAnalysisSubmenu", "Analysis")
+	};
+
+	return SubMenus;
 }
 
 UClass* UAudioSynesthesiaNRT::GetSupportedClass() const
diff --git a/Engine/Plugins/Runtime/AudioSynesthesia/Source/AudioSynesthesiaEditor/Private/AssetTypeActions_AudioSynesthesiaNRT.cpp b/Engine/Plugins/Runtime/AudioSynesthesia/Source/AudioSynesthesiaEditor/Private/AssetTypeActions_AudioSynesthesiaNRT.cpp
--- a/Engine/Plugins/Runtime/AudioSynesthesia/Source/AudioSynesthesiaEditor/Private/AssetTypeActions_AudioSynesthesiaNRT.cpp
+++ b/Engine/Plugins/Runtime/AudioSynesthesia/Source/AudioSynesthesiaEditor/Private/AssetTypeActions_AudioSynesthesiaNRT.cpp
@@ -24,6 +24,11 @@ FText FAssetTypeActions_AudioSynesthesiaNRT::GetName() const
 	}
 }
 
+const TArray<FText>& FAssetTypeActions_AudioSynesthesiaNRT::GetSubMenus() const
+{
+	return Synesthesia->GetAssetActionSubmenus();
+}
+
 FColor FAssetTypeActions_AudioSynesthesiaNRT::GetTypeColor() const 
 {
 	return Synesthesia->GetTypeColor(); 
@@ -46,4 +51,3 @@ uint32 FAssetTypeActions_AudioSynesthesiaNRT::GetCategories()
 {
 	return EAssetTypeCategories::Sounds; 
 }
-
diff --git a/Engine/Plugins/Runtime/AudioSynesthesia/Source/AudioSynesthesiaEditor/Private/AssetTypeActions_AudioSynesthesiaNRT.h b/Engine/Plugins/Runtime/AudioSynesthesia/Source/AudioSynesthesiaEditor/Private/AssetTypeActions_AudioSynesthesiaNRT.h
--- a/Engine/Plugins/Runtime/AudioSynesthesia/Source/AudioSynesthesiaEditor/Private/AssetTypeActions_AudioSynesthesiaNRT.h
+++ b/Engine/Plugins/Runtime/AudioSynesthesia/Source/AudioSynesthesiaEditor/Private/AssetTypeActions_AudioSynesthesiaNRT.h
@@ -18,6 +18,7 @@ public:
 	virtual FColor GetTypeColor() const override;
 	virtual UClass* GetSupportedClass() const override;
 	virtual uint32 GetCategories() override;
+	virtual const TArray<FText>& GetSubMenus() const override;
 	//~ End FAssetTypeActions_Base
 
 private:
diff --git a/Engine/Plugins/Runtime/AudioSynesthesia/Source/AudioSynesthesiaEditor/Private/AssetTypeActions_AudioSynesthesiaNRTSettings.cpp b/Engine/Plugins/Runtime/AudioSynesthesia/Source/AudioSynesthesiaEditor/Private/AssetTypeActions_AudioSynesthesiaNRTSettings.cpp
--- a/Engine/Plugins/Runtime/AudioSynesthesia/Source/AudioSynesthesiaEditor/Private/AssetTypeActions_AudioSynesthesiaNRTSettings.cpp
+++ b/Engine/Plugins/Runtime/AudioSynesthesia/Source/AudioSynesthesiaEditor/Private/AssetTypeActions_AudioSynesthesiaNRTSettings.cpp
@@ -26,7 +26,12 @@ FText FAssetTypeActions_AudioSynesthesiaNRTSettings::GetName() const
 
 FColor FAssetTypeActions_AudioSynesthesiaNRTSettings::GetTypeColor() const 
 {
-	return SynesthesiaSettings->GetTypeColor(); 
+	return SynesthesiaSettings->GetTypeColor();
+}
+
+const TArray<FText>& FAssetTypeActions_AudioSynesthesiaNRTSettings::GetSubMenus() const
+{
+	return SynesthesiaSettings->GetAssetActionSubmenus();
 }
 
 UClass* FAssetTypeActions_AudioSynesthesiaNRTSettings::GetSupportedClass() const
diff --git a/Engine/Plugins/Runtime/AudioSynesthesia/Source/AudioSynesthesiaEditor/Private/AssetTypeActions_AudioSynesthesiaNRTSettings.h b/Engine/Plugins/Runtime/AudioSynesthesia/Source/AudioSynesthesiaEditor/Private/AssetTypeActions_AudioSynesthesiaNRTSettings.h
--- a/Engine/Plugins/Runtime/AudioSynesthesia/Source/AudioSynesthesiaEditor/Private/AssetTypeActions_AudioSynesthesiaNRTSettings.h
+++ b/Engine/Plugins/Runtime/AudioSynesthesia/Source/AudioSynesthesiaEditor/Private/AssetTypeActions_AudioSynesthesiaNRTSettings.h
@@ -18,6 +18,7 @@ public:
 	virtual FColor GetTypeColor() const override;
 	virtual UClass* GetSupportedClass() const override;
 	virtual uint32 GetCategories() override;
+	virtual const TArray<FText>& GetSubMenus() const override;
 	//~ End FAssetTypeActions_Base
 
 private:
diff --git a/Engine/Plugins/Runtime/ResonanceAudio/Source/ResonanceAudioEditor/Private/ResonanceAudioEditorModule.cpp b/Engine/Plugins/Runtime/ResonanceAudio/Source/ResonanceAudioEditor/Private/ResonanceAudioEditorModule.cpp
--- a/Engine/Plugins/Runtime/ResonanceAudio/Source/ResonanceAudioEditor/Private/ResonanceAudioEditorModule.cpp
+++ b/Engine/Plugins/Runtime/ResonanceAudio/Source/ResonanceAudioEditor/Private/ResonanceAudioEditorModule.cpp
@@ -51,8 +51,8 @@ namespace ResonanceAudio
 		// Register the audio editor asset type actions.
 		IAssetTools& AssetTools = FModuleManager::LoadModuleChecked<FAssetToolsModule>("AssetTools").Get();
 
-		AssetTools.RegisterAssetTypeActions(MakeShareable(new FAssetTypeActions_ResonanceAudioReverbPluginPreset));
-		AssetTools.RegisterAssetTypeActions(MakeShareable(new FAssetTypeActions_ResonanceAudioSpatializationSourceSettings));
+		AssetTools.RegisterAssetTypeActions(MakeShared<FAssetTypeActions_ResonanceAudioReverbPluginPreset>());
+		AssetTools.RegisterAssetTypeActions(MakeShared<FAssetTypeActions_ResonanceAudioSpatializationSourceSettings>());
 	}
 
 	void FResonanceAudioEditorModule::ShutdownModule()
diff --git a/Engine/Plugins/Runtime/ResonanceAudio/Source/ResonanceAudioEditor/Private/ResonanceAudioReverbPluginPresetFactory.cpp b/Engine/Plugins/Runtime/ResonanceAudio/Source/ResonanceAudioEditor/Private/ResonanceAudioReverbPluginPresetFactory.cpp
--- a/Engine/Plugins/Runtime/ResonanceAudio/Source/ResonanceAudioEditor/Private/ResonanceAudioReverbPluginPresetFactory.cpp
+++ b/Engine/Plugins/Runtime/ResonanceAudio/Source/ResonanceAudioEditor/Private/ResonanceAudioReverbPluginPresetFactory.cpp
@@ -27,6 +27,16 @@ uint32 FAssetTypeActions_ResonanceAudioReverbPluginPreset::GetCategories()
 	return EAssetTypeCategories::Sounds;
 }
 
+const TArray<FText>& FAssetTypeActions_ResonanceAudioReverbPluginPreset::GetSubMenus() const
+{
+	static const TArray<FText> SubMenus
+	{
+		NSLOCTEXT("AssetTypeActions", "AssetTypeActions_AssetSoundResonanceSubMenu", "Resonance")
+	};
+
+	return SubMenus;
+}
+
 UResonanceAudioReverbPluginPresetFactory::UResonanceAudioReverbPluginPresetFactory(const FObjectInitializer& ObjectInitializer)
 	: Super(ObjectInitializer)
 {
diff --git a/Engine/Plugins/Runtime/ResonanceAudio/Source/ResonanceAudioEditor/Private/ResonanceAudioReverbPluginPresetFactory.h b/Engine/Plugins/Runtime/ResonanceAudio/Source/ResonanceAudioEditor/Private/ResonanceAudioReverbPluginPresetFactory.h
--- a/Engine/Plugins/Runtime/ResonanceAudio/Source/ResonanceAudioEditor/Private/ResonanceAudioReverbPluginPresetFactory.h
+++ b/Engine/Plugins/Runtime/ResonanceAudio/Source/ResonanceAudioEditor/Private/ResonanceAudioReverbPluginPresetFactory.h
@@ -15,6 +15,7 @@ public:
 	virtual FColor GetTypeColor() const override;
 	virtual UClass* GetSupportedClass() const override;
 	virtual uint32 GetCategories() override;
+	virtual const TArray<FText>& GetSubMenus() const override;
 };
 
 UCLASS(MinimalAPI, hidecategories = Object)
diff --git a/Engine/Plugins/Runtime/ResonanceAudio/Source/ResonanceAudioEditor/Private/ResonanceAudioSpatializationSourceSettingsFactory.cpp b/Engine/Plugins/Runtime/ResonanceAudio/Source/ResonanceAudioEditor/Private/ResonanceAudioSpatializationSourceSettingsFactory.cpp
--- a/Engine/Plugins/Runtime/ResonanceAudio/Source/ResonanceAudioEditor/Private/ResonanceAudioSpatializationSourceSettingsFactory.cpp
+++ b/Engine/Plugins/Runtime/ResonanceAudio/Source/ResonanceAudioEditor/Private/ResonanceAudioSpatializationSourceSettingsFactory.cpp
@@ -27,6 +27,16 @@ uint32 FAssetTypeActions_ResonanceAudioSpatializationSourceSettings::GetCategori
 	return EAssetTypeCategories::Sounds;
 }
 
+const TArray<FText>& FAssetTypeActions_ResonanceAudioSpatializationSourceSettings::GetSubMenus() const
+{
+	static const TArray<FText> SubMenus
+	{
+		NSLOCTEXT("AssetTypeActions", "AssetTypeActions_AssetSoundResonanceSubMenu", "Resonance")
+	};
+
+	return SubMenus;
+}
+
 UResonanceAudioSpatializationSourceSettingsFactory::UResonanceAudioSpatializationSourceSettingsFactory(const FObjectInitializer& ObjectInitializer)
 	: Super(ObjectInitializer)
 {
diff --git a/Engine/Plugins/Runtime/ResonanceAudio/Source/ResonanceAudioEditor/Private/ResonanceAudioSpatializationSourceSettingsFactory.h b/Engine/Plugins/Runtime/ResonanceAudio/Source/ResonanceAudioEditor/Private/ResonanceAudioSpatializationSourceSettingsFactory.h
--- a/Engine/Plugins/Runtime/ResonanceAudio/Source/ResonanceAudioEditor/Private/ResonanceAudioSpatializationSourceSettingsFactory.h
+++ b/Engine/Plugins/Runtime/ResonanceAudio/Source/ResonanceAudioEditor/Private/ResonanceAudioSpatializationSourceSettingsFactory.h
@@ -15,6 +15,7 @@ public:
 	virtual FColor GetTypeColor() const override;
 	virtual UClass* GetSupportedClass() const override;
 	virtual uint32 GetCategories() override;
+	virtual const TArray<FText>& GetSubMenus() const override;
 };
 
 UCLASS(MinimalAPI, hidecategories = Object)
diff --git a/Engine/Source/Runtime/Engine/Public/IAudioExtensionPlugin.h b/Engine/Source/Runtime/Engine/Public/IAudioExtensionPlugin.h
--- a/Engine/Source/Runtime/Engine/Public/IAudioExtensionPlugin.h
+++ b/Engine/Source/Runtime/Engine/Public/IAudioExtensionPlugin.h
@@ -461,7 +461,7 @@ public:
 /** Interface to sound that is modulateable, allowing for certain specific 
   * behaviors to be controlled on the sound level by the modulation system.
   */
-class ISoundModulatable
+class ENGINE_API ISoundModulatable
 {
 public:
 	virtual ~ISoundModulatable() = default;
@@ -529,7 +529,7 @@ public:
  * Modulatable controls found on each sound instance
  * processed by the modulation plugin enabled
  */
-struct FSoundModulationControls
+struct ENGINE_API FSoundModulationControls
 {
 	float Volume;
 	float Pitch;
