Audio Synesthesia Updates

Added onset analyzer which exposes audio onsets to blueprints.
Updated audio analysis trigger mechanism to avoid triggering analysis during an interaction
Updated saving mechanism to mark uobjects dirty when new results generated.
Updated audio analysis logic to keep track of most recent analysis result and avoid ordering issues arrising from async processing

diff --git a/Engine/Plugins/Runtime/AudioSynesthesia/Source/AudioSynesthesia/Classes/ConstantQNRT.h b/Engine/Plugins/Runtime/AudioSynesthesia/Source/AudioSynesthesia/Classes/ConstantQNRT.h
--- a/Engine/Plugins/Runtime/AudioSynesthesia/Source/AudioSynesthesia/Classes/ConstantQNRT.h
+++ b/Engine/Plugins/Runtime/AudioSynesthesia/Source/AudioSynesthesia/Classes/ConstantQNRT.h
@@ -71,51 +71,51 @@ class AUDIOSYNESTHESIA_API UConstantQNRTSettings : public UAudioSynesthesiaNRTSe
 		UConstantQNRTSettings();
 
 		/** Starting frequency for first bin of CQT */
-		UPROPERTY(EditAnywhere, BlueprintReadWrite, Category=AudioAnalyzer, meta = (ClampMin = "20.0", ClampMax = "20000"))
+		UPROPERTY(EditAnywhere, BlueprintReadOnly, Category=AudioAnalyzer, meta = (ClampMin = "20.0", ClampMax = "20000"))
 		float StartingFrequency;
 
         /** Total number of resulting constant Q bands. */
-		UPROPERTY(EditAnywhere, BlueprintReadWrite, Category=AudioAnalyzer, meta = (ClampMin = "1", ClampMax = "96"))
+		UPROPERTY(EditAnywhere, BlueprintReadOnly, Category=AudioAnalyzer, meta = (ClampMin = "1", ClampMax = "96"))
 		int32 NumBands;
 
 		/** Number of bands within an octave. */
-		UPROPERTY(EditAnywhere, BlueprintReadWrite, Category=AudioAnalyzer, meta = (ClampMin = "1", ClampMax = "24"))
+		UPROPERTY(EditAnywhere, BlueprintReadOnly, Category=AudioAnalyzer, meta = (ClampMin = "1", ClampMax = "24"))
 		float NumBandsPerOctave;                        
 
 		/** Number of seconds between cqt measurements */
-		UPROPERTY(EditAnywhere, BlueprintReadWrite, Category=AudioAnalyzer, meta = (ClampMin = "0.01", ClampMax = "1.0"))
+		UPROPERTY(EditAnywhere, BlueprintReadOnly, Category=AudioAnalyzer, meta = (ClampMin = "0.01", ClampMax = "1.0"))
 		float AnalysisPeriod;
 
 		/** If true, multichannel audio is downmixed to mono with equal amplitude scaling. If false, each channel gets it's own CQT result. */
-		UPROPERTY(EditAnywhere, BlueprintReadWrite, Category=AudioAnalyzer)
+		UPROPERTY(EditAnywhere, BlueprintReadOnly, Category=AudioAnalyzer)
 		bool bDownmixToMono;
 
 		/** Size of FFT. */
-		UPROPERTY(EditAnywhere, BlueprintReadWrite, AdvancedDisplay, Category=AudioAnalyzer)
+		UPROPERTY(EditAnywhere, BlueprintReadOnly, AdvancedDisplay, Category=AudioAnalyzer)
 		EConstantQFFTSizeEnum FFTSize;
 
 		/** Type of window to be applied to input audio */
-		UPROPERTY(EditAnywhere, BlueprintReadWrite, AdvancedDisplay, Category=AudioAnalyzer)
+		UPROPERTY(EditAnywhere, BlueprintReadOnly, AdvancedDisplay, Category=AudioAnalyzer)
         EFFTWindowType WindowType;
 
 		/** Type of spectrum to use. */
-		UPROPERTY(EditAnywhere, BlueprintReadWrite, AdvancedDisplay, Category=AudioAnalyzer)
+		UPROPERTY(EditAnywhere, BlueprintReadOnly, AdvancedDisplay, Category=AudioAnalyzer)
 		EAudioSpectrumType SpectrumType;
 		
 		/** Stretching factor to control overlap of adjacent bands. */
-		UPROPERTY(EditAnywhere, BlueprintReadWrite, AdvancedDisplay, Category=AudioAnalyzer, meta = (ClampMin = "0.01", ClampMax = "2.0"))
+		UPROPERTY(EditAnywhere, BlueprintReadOnly, AdvancedDisplay, Category=AudioAnalyzer, meta = (ClampMin = "0.01", ClampMax = "2.0"))
 		float BandWidthStretch;                         
 		
 		/** Normalization scheme used to generate band windows. */
-		UPROPERTY(EditAnywhere, BlueprintReadWrite, AdvancedDisplay, Category=AudioAnalyzer)
+		UPROPERTY(EditAnywhere, BlueprintReadOnly, AdvancedDisplay, Category=AudioAnalyzer)
 		EConstantQNormalizationEnum CQTNormalization;
 
 		/** Noise floor to use when normalizing CQT */
-		UPROPERTY(EditAnywhere, BlueprintReadWrite, AdvancedDisplay, Category = AudioAnalyzer, meta = (ClampMin = "-120.0", ClampMax="0.0"))
+		UPROPERTY(EditAnywhere, BlueprintReadOnly, AdvancedDisplay, Category = AudioAnalyzer, meta = (ClampMin = "-120.0", ClampMax="0.0"))
 		float NoiseFloorDb;
 
 		/** Convert UConstantQNRTSettings to FConstantQNRTSettings */
-		TUniquePtr<Audio::IAnalyzerNRTSettings> GetSettings();
+		TUniquePtr<Audio::IAnalyzerNRTSettings> GetSettings(const float InSampleRate, const int32 InNumChannels) const;
 };
 
 
@@ -134,19 +134,19 @@ class AUDIOSYNESTHESIA_API UConstantQNRT : public UAudioSynesthesiaNRT
 		UConstantQNRT();
 
 		/** The settings for the audio analyzer.  */
-		UPROPERTY(EditAnywhere, BlueprintReadWrite, Category=AudioAnalyzer)
+		UPROPERTY(EditAnywhere, BlueprintReadOnly, Category=AudioAnalyzer)
 		UConstantQNRTSettings* Settings;
 
 		/** Get a specific channel cqt of the analyzed sound at a given time. */
 		UFUNCTION(BlueprintCallable, Category="Audio Analyzer")
-		void GetChannelConstantQAtTime(const float InSeconds, const int32 InChannel, TArray<float>& OutConstantQ);
+		void GetChannelConstantQAtTime(const float InSeconds, const int32 InChannel, TArray<float>& OutConstantQ) const;
 
 		/** Get a specific channel cqt of the analyzed sound at a given time. */
 		UFUNCTION(BlueprintCallable, Category="Audio Analyzer")
-		void GetNormalizedChannelConstantQAtTime(const float InSeconds, const int32 InChannel, TArray<float>& OutConstantQ);
+		void GetNormalizedChannelConstantQAtTime(const float InSeconds, const int32 InChannel, TArray<float>& OutConstantQ) const;
 
 		/** Convert ULoudnessNRTSettings to FLoudnessNRTSettings */
- 		virtual TUniquePtr<Audio::IAnalyzerNRTSettings> GetSettings() override;
+ 		virtual TUniquePtr<Audio::IAnalyzerNRTSettings> GetSettings(const float InSampleRate, const int32 InNumChannels) const override;
 
 	protected:
 
diff --git a/Engine/Plugins/Runtime/AudioSynesthesia/Source/AudioSynesthesia/Classes/LoudnessNRT.h b/Engine/Plugins/Runtime/AudioSynesthesia/Source/AudioSynesthesia/Classes/LoudnessNRT.h
--- a/Engine/Plugins/Runtime/AudioSynesthesia/Source/AudioSynesthesia/Classes/LoudnessNRT.h
+++ b/Engine/Plugins/Runtime/AudioSynesthesia/Source/AudioSynesthesia/Classes/LoudnessNRT.h
@@ -37,27 +37,27 @@ class AUDIOSYNESTHESIA_API ULoudnessNRTSettings : public UAudioSynesthesiaNRTSet
 		ULoudnessNRTSettings();
 
 		/** Number of seconds between loudness measurements */
-		UPROPERTY(EditAnywhere, BlueprintReadWrite, Category=AudioAnalyzer, meta = (ClampMin = "0.01", ClampMax = "0.25"))
+		UPROPERTY(EditAnywhere, BlueprintReadOnly, Category=AudioAnalyzer, meta = (ClampMin = "0.01", ClampMax = "0.25"))
 		float AnalysisPeriod;
 
 		/** Minimum analysis frequency for calculating loudness. */
-		UPROPERTY(EditAnywhere, BlueprintReadWrite, Category=AudioAnalyzer, meta = (ClampMin = "20.0", ClampMax = "20000"))
+		UPROPERTY(EditAnywhere, BlueprintReadOnly, Category=AudioAnalyzer, meta = (ClampMin = "20.0", ClampMax = "20000"))
 		float MinimumFrequency;
 
 		/** Maximum analysis frequency for calculating loudness. */
-		UPROPERTY(EditAnywhere, BlueprintReadWrite, Category=AudioAnalyzer, meta = (ClampMin = "20.0", ClampMax = "20000"))
+		UPROPERTY(EditAnywhere, BlueprintReadOnly, Category=AudioAnalyzer, meta = (ClampMin = "20.0", ClampMax = "20000"))
 		float MaximumFrequency;
 
 		/** Type of equal loudness curve used in calculations */
-		UPROPERTY(EditAnywhere, BlueprintReadWrite, Category=AudioAnalyzer)
+		UPROPERTY(EditAnywhere, BlueprintReadOnly, Category=AudioAnalyzer)
 		ELoudnessNRTCurveTypeEnum CurveType;
 
 		/** dB level to denote silence.  Used when calculating normalized loudness. */
-		UPROPERTY(EditAnywhere, BlueprintReadWrite, AdvancedDisplay, Category=AudioAnalyzer, meta = (ClampMin = "-100.0", ClampMax="0"))
+		UPROPERTY(EditAnywhere, BlueprintReadOnly, AdvancedDisplay, Category=AudioAnalyzer, meta = (ClampMin = "-100.0", ClampMax="0"))
 		float NoiseFloorDb;
 
 		/** Convert ULoudnessNRTSettings to FLoudnessNRTSettings */
-		TUniquePtr<Audio::IAnalyzerNRTSettings> GetSettings();
+		TUniquePtr<Audio::IAnalyzerNRTSettings> GetSettings(const float InSampleRate, const int32 InNumChannels) const;
 };
 
 
@@ -78,16 +78,16 @@ class AUDIOSYNESTHESIA_API ULoudnessNRT : public UAudioSynesthesiaNRT
 		ULoudnessNRT();
 
 		/** The settings for the audio analyzer.  */
-		UPROPERTY(EditAnywhere, BlueprintReadWrite, Category=AudioAnalyzer)
+		UPROPERTY(EditAnywhere, BlueprintReadOnly, Category=AudioAnalyzer)
 		ULoudnessNRTSettings* Settings;
 
 		/** Get the overall loudness of the analyzed sound at a given time. */
 		UFUNCTION(BlueprintCallable, Category="Audio Analyzer")
-		void GetLoudnessAtTime(const float InSeconds, float& OutLoudness);
+		void GetLoudnessAtTime(const float InSeconds, float& OutLoudness) const;
 
 		/** Get a specific channel loudness of the analyzed sound at a given time. */
 		UFUNCTION(BlueprintCallable, Category="Audio Analyzer")
-		void GetChannelLoudnessAtTime(const float InSeconds, const int32 InChannel, float& OutLoudness);
+		void GetChannelLoudnessAtTime(const float InSeconds, const int32 InChannel, float& OutLoudness) const;
 
 		/**
 		 * Get the normalized overall loudness of the analyzed sound at a given time. Normalized loudness
@@ -95,7 +95,7 @@ class AUDIOSYNESTHESIA_API ULoudnessNRT : public UAudioSynesthesiaNRT
 		 * loudness in the sound.
 		 */
 		UFUNCTION(BlueprintCallable, Category="Audio Analyzer")
-		void GetNormalizedLoudnessAtTime(const float InSeconds, float& OutLoudness);
+		void GetNormalizedLoudnessAtTime(const float InSeconds, float& OutLoudness) const;
 
 		/**
 		 * Get a specific channel normalized loudness of the analyzed sound at a given time. Normalized 
@@ -103,10 +103,10 @@ class AUDIOSYNESTHESIA_API ULoudnessNRT : public UAudioSynesthesiaNRT
 		 * maximum loudness in the sound.
 		 */
 		UFUNCTION(BlueprintCallable, Category="Audio Analyzer")
-		void GetNormalizedChannelLoudnessAtTime(const float InSeconds, const int32 InChannel, float& OutLoudness);
+		void GetNormalizedChannelLoudnessAtTime(const float InSeconds, const int32 InChannel, float& OutLoudness) const;
 
 		/** Convert ULoudnessNRTSettings to FLoudnessNRTSettings */
- 		TUniquePtr<Audio::IAnalyzerNRTSettings> GetSettings() override;
+ 		TUniquePtr<Audio::IAnalyzerNRTSettings> GetSettings(const float InSampleRate, const int32 InNumChannels) const override;
 
 	protected:
 
diff --git a/Engine/Plugins/Runtime/AudioSynesthesia/Source/AudioSynesthesia/Classes/OnsetNRT.h b/Engine/Plugins/Runtime/AudioSynesthesia/Source/AudioSynesthesia/Classes/OnsetNRT.h
new file mode 100644
--- /dev/null
+++ b/Engine/Plugins/Runtime/AudioSynesthesia/Source/AudioSynesthesia/Classes/OnsetNRT.h
@@ -0,0 +1,81 @@
+// Copyright 1998-2019 Epic Games, Inc. All Rights Reserved.
+
+#pragma once
+
+#include "AudioSynesthesiaNRT.h"
+#include "OnsetNRT.generated.h"
+										 
+/** UOnsetNRTSettings
+ *
+ * Settings for a UOnsetNRT analyzer.
+ */
+UCLASS(Blueprintable)
+class AUDIOSYNESTHESIA_API UOnsetNRTSettings : public UAudioSynesthesiaNRTSettings
+{
+	GENERATED_BODY()
+	public:
+
+		UOnsetNRTSettings();
+
+		/** If true, multichannel audio is downmixed to mono with equal amplitude scaling. If false, each channel gets it's own onset result. */
+		UPROPERTY(EditAnywhere, BlueprintReadOnly, Category=AudioAnalyzer)
+		bool bDownmixToMono;
+
+		/** Onset timestamp granularity onsets. Lower granularity takes longer to compute. */
+		UPROPERTY(EditAnywhere, BlueprintReadOnly, Category=AudioAnalyzer, meta = (ClampMin = "0.005", ClampMax = "0.25"))
+		float GranularityInSeconds;
+
+		/** Sensitivity of onset detector. Higher sensitivity will find more onsets. */
+		UPROPERTY(EditAnywhere, BlueprintReadOnly, Category=AudioAnalyzer, meta = (ClampMin = "0.0", ClampMax = "1.0"))
+		float Sensitivity;
+
+		/** Starting frequency for onset anlaysis. */
+		UPROPERTY(EditAnywhere, BlueprintReadOnly, Category=AudioAnalyzer, meta = (ClampMin = "20.0", ClampMax = "20000"))
+		float MinimumFrequency;
+
+		/** Starting frequency for onset anlaysis. */
+		UPROPERTY(EditAnywhere, BlueprintReadOnly, Category=AudioAnalyzer, meta = (ClampMin = "20.0", ClampMax = "20000"))
+		float MaximumFrequency;
+
+
+		/** Convert UOnsetNRTSettings to FOnsetNRTSettings */
+		TUniquePtr<Audio::IAnalyzerNRTSettings> GetSettings(const float InSampleRate, const int32 InNumChannels) const;
+};
+
+
+/** UOnsetNRT
+ *
+ * UOnsetNRT calculates the temporal evolution of constant q transform for a given 
+ * sound. Onset is available for individual channels or the overall sound asset.
+ */
+UCLASS(Blueprintable)
+class AUDIOSYNESTHESIA_API UOnsetNRT : public UAudioSynesthesiaNRT
+{
+	GENERATED_BODY()
+
+	public:
+
+		UOnsetNRT();
+
+		/** The settings for the audio analyzer.  */
+		UPROPERTY(EditAnywhere, BlueprintReadOnly, Category=AudioAnalyzer)
+		UOnsetNRTSettings* Settings;
+
+		/** Returns onsets which occured between start and end timestamps. */
+		UFUNCTION(BlueprintCallable, Category="Audio Analyzer")
+		void GetChannelOnsetsBetweenTimes(const float InStartSeconds, const float InEndSeconds, const int32 InChannel, TArray<float>& OutOnsetTimestamps, TArray<float>& OutOnsetStrengths) const;
+
+		/** Get a specific channel cqt of the analyzed sound at a given time. */
+		UFUNCTION(BlueprintCallable, Category="Audio Analyzer")
+		void GetNormalizedChannelOnsetsBetweenTimes(const float InStartSeconds, const float InEndSeconds, const int32 InChannel, TArray<float>& OutOnsetTimestamps, TArray<float>& OutOnsetStrengths) const;
+
+		/** Convert ULoudnessNRTSettings to FLoudnessNRTSettings */
+ 		virtual TUniquePtr<Audio::IAnalyzerNRTSettings> GetSettings(const float InSampleRate, const int32 InNumChannels) const override;
+
+	protected:
+
+		/** Return the name of the IAudioAnalyzerNRTFactory associated with this UAudioAnalyzerNRT */
+		virtual FName GetAnalyzerNRTFactoryName() const override;
+
+};
+
diff --git a/Engine/Plugins/Runtime/AudioSynesthesia/Source/AudioSynesthesia/Private/AudioSynesthesiaCustomVersion.cpp b/Engine/Plugins/Runtime/AudioSynesthesia/Source/AudioSynesthesia/Private/AudioSynesthesiaCustomVersion.cpp
new file mode 100644
--- /dev/null
+++ b/Engine/Plugins/Runtime/AudioSynesthesia/Source/AudioSynesthesia/Private/AudioSynesthesiaCustomVersion.cpp
@@ -0,0 +1,9 @@
+// Copyright 1998-2019 Epic Games, Inc. All Rights Reserved.
+
+#include "AudioSynesthesiaCustomVersion.h"
+#include "Serialization/CustomVersion.h"
+
+const FGuid FAudioSynesthesiaCustomVersion::GUID( 0x82E77C4E, 0x332343A5, 0xB46B13C5, 0x97310DF3 );
+
+// Register the custom version with core
+FCustomVersionRegistration GRegisterAudioSynesthesiaCustomVersion( FAudioSynesthesiaCustomVersion::GUID, FAudioSynesthesiaCustomVersion::LatestVersion, TEXT( "AudioSynesthesiaVer" ) );
diff --git a/Engine/Plugins/Runtime/AudioSynesthesia/Source/AudioSynesthesia/Private/AudioSynesthesiaCustomVersion.h b/Engine/Plugins/Runtime/AudioSynesthesia/Source/AudioSynesthesia/Private/AudioSynesthesiaCustomVersion.h
new file mode 100644
--- /dev/null
+++ b/Engine/Plugins/Runtime/AudioSynesthesia/Source/AudioSynesthesia/Private/AudioSynesthesiaCustomVersion.h
@@ -0,0 +1,25 @@
+// Copyright 1998-2019 Epic Games, Inc. All Rights Reserved.
+
+#pragma once
+
+#include "CoreMinimal.h"
+#include "Misc/Guid.h"
+
+// Custom serialization version for all packages containing AudioSynesthesia asset types
+struct FAudioSynesthesiaCustomVersion
+{
+	enum Type
+	{
+		// Before any version changes were made in the module
+		BeforeCustomVersionWasAdded = 0,
+
+		// -----<new versions can be added above this line>-------------------------------------------------
+		VersionPlusOne,
+		LatestVersion = VersionPlusOne - 1
+	};
+
+	// The GUID for this custom version number
+	const static FGuid GUID;
+
+	FAudioSynesthesiaCustomVersion() = delete;
+};
diff --git a/Engine/Plugins/Runtime/AudioSynesthesia/Source/AudioSynesthesia/Private/AudioSynesthesiaLog.h b/Engine/Plugins/Runtime/AudioSynesthesia/Source/AudioSynesthesia/Private/AudioSynesthesiaLog.h
new file mode 100644
--- /dev/null
+++ b/Engine/Plugins/Runtime/AudioSynesthesia/Source/AudioSynesthesia/Private/AudioSynesthesiaLog.h
@@ -0,0 +1,7 @@
+// Copyright 1998-2019 Epic Games, Inc. All Rights Reserved.
+
+#pragma once
+
+#include "CoreMinimal.h"
+
+DECLARE_LOG_CATEGORY_EXTERN(LogAudioSynesthesia, Log, All);
diff --git a/Engine/Plugins/Runtime/AudioSynesthesia/Source/AudioSynesthesia/Private/AudioSynesthesiaModule.cpp b/Engine/Plugins/Runtime/AudioSynesthesia/Source/AudioSynesthesia/Private/AudioSynesthesiaModule.cpp
--- a/Engine/Plugins/Runtime/AudioSynesthesia/Source/AudioSynesthesia/Private/AudioSynesthesiaModule.cpp
+++ b/Engine/Plugins/Runtime/AudioSynesthesia/Source/AudioSynesthesia/Private/AudioSynesthesiaModule.cpp
@@ -5,6 +5,7 @@
 #include "Modules/ModuleInterface.h"
 #include "Features/IModularFeatures.h"
 #include "LoudnessNRTFactory.h"
+#include "AudioSynesthesiaLog.h"
 
 DEFINE_LOG_CATEGORY(LogAudioSynesthesia);
 
@@ -16,11 +17,15 @@ namespace Audio
 	{
 		// Register factories on startup
 		IModularFeatures::Get().RegisterModularFeature(FLoudnessNRTFactory::GetModularFeatureName(), &LoudnessFactory);
+		IModularFeatures::Get().RegisterModularFeature(FConstantQNRTFactory::GetModularFeatureName(), &ConstantQFactory);
+		IModularFeatures::Get().RegisterModularFeature(FOnsetNRTFactory::GetModularFeatureName(), &OnsetFactory);
 	}
 
 	void FAudioSynesthesiaModule::ShutdownModule()
 	{
 		// Unregister factories on shutdown
 		IModularFeatures::Get().UnregisterModularFeature(FLoudnessNRTFactory::GetModularFeatureName(), &LoudnessFactory);
+		IModularFeatures::Get().UnregisterModularFeature(FConstantQNRTFactory::GetModularFeatureName(), &ConstantQFactory);
+		IModularFeatures::Get().UnregisterModularFeature(FOnsetNRTFactory::GetModularFeatureName(), &OnsetFactory);
 	}
 }
diff --git a/Engine/Plugins/Runtime/AudioSynesthesia/Source/AudioSynesthesia/Private/ConstantQNRT.cpp b/Engine/Plugins/Runtime/AudioSynesthesia/Source/AudioSynesthesia/Private/ConstantQNRT.cpp
--- a/Engine/Plugins/Runtime/AudioSynesthesia/Source/AudioSynesthesia/Private/ConstantQNRT.cpp
+++ b/Engine/Plugins/Runtime/AudioSynesthesia/Source/AudioSynesthesia/Private/ConstantQNRT.cpp
@@ -2,9 +2,11 @@
 
 #include "ConstantQNRT.h"
 #include "ConstantQNRTFactory.h"
+#include "AudioSynesthesiaLog.h"
 #include "InterpolateSorted.h"
 #include "DSP/ConstantQ.h"
 
+
 namespace
 {
 	/** Convenience template function for converting from blueprint types to audio types */
@@ -133,7 +135,7 @@ UConstantQNRTSettings::UConstantQNRTSettings()
 {}
 
 /** Convert UConstantQNRTSettings to FConstantQNRTSettings */
-TUniquePtr<Audio::IAnalyzerNRTSettings> UConstantQNRTSettings::GetSettings()
+TUniquePtr<Audio::IAnalyzerNRTSettings> UConstantQNRTSettings::GetSettings(const float InSampleRate, const int32 InNumChannels) const
 {
 	TUniquePtr<Audio::FConstantQNRTSettings> Settings = MakeUnique<Audio::FConstantQNRTSettings>();
 
@@ -163,22 +165,26 @@ UConstantQNRT::UConstantQNRT()
 
 #if WITH_EDITOR
 	// Bind settings to audio analyze so changes to default settings will trigger analysis.
-	Settings->AnalyzeAudioDelegate.BindUObject(this, &UAudioAnalyzerNRT::AnalyzeAudio);
+	SetSettingsDelegate(Settings);
 #endif
 }
 
 
-void UConstantQNRT::GetChannelConstantQAtTime(const float InSeconds, const int32 InChannel, TArray<float>& OutConstantQ)
+void UConstantQNRT::GetChannelConstantQAtTime(const float InSeconds, const int32 InChannel, TArray<float>& OutConstantQ) const
 {
 	OutConstantQ.Reset();
 
-	TSharedPtr<Audio::FConstantQNRTResult, ESPMode::ThreadSafe> ConstantQResult = GetResult<Audio::FConstantQNRTResult>();
+	TSharedPtr<const Audio::FConstantQNRTResult, ESPMode::ThreadSafe> ConstantQResult = GetResult<Audio::FConstantQNRTResult>();
 
 	if (ConstantQResult.IsValid())
 	{
-		if (!ConstantQResult->IsSortedChronologically())
+		// ConstantQResults should always be sorted if in UConstantQNRT
+		check(ConstantQResult->IsSortedChronologically());
+
+		if (!ConstantQResult->ContainsChannel(InChannel))
 		{
-			ConstantQResult->SortChronologically();
+			UE_LOG(LogAudioSynesthesia, Warning, TEXT("ConstantQNRT does not contain channel %d"), InChannel);
+			return;
 		}
 
 		const TArray<Audio::FConstantQFrame>& ConstantQArray = ConstantQResult->GetFramesForChannel(InChannel);
@@ -187,17 +193,21 @@ void UConstantQNRT::GetChannelConstantQAtTime(const float InSeconds, const int32
 	}
 }
 
-void UConstantQNRT::GetNormalizedChannelConstantQAtTime(const float InSeconds, const int32 InChannelIdx, TArray<float>& OutConstantQ)
+void UConstantQNRT::GetNormalizedChannelConstantQAtTime(const float InSeconds, const int32 InChannelIdx, TArray<float>& OutConstantQ) const
 {
 	OutConstantQ.Reset();
 
-	TSharedPtr<Audio::FConstantQNRTResult, ESPMode::ThreadSafe> ConstantQResult = GetResult<Audio::FConstantQNRTResult>();
+	TSharedPtr<const Audio::FConstantQNRTResult, ESPMode::ThreadSafe> ConstantQResult = GetResult<Audio::FConstantQNRTResult>();
 
 	if (ConstantQResult.IsValid())
 	{
-		if (!ConstantQResult->IsSortedChronologically())
+		// ConstantQResults should always be sorted if in UConstantQNRT
+		check(ConstantQResult->IsSortedChronologically());
+
+		if (!ConstantQResult->ContainsChannel(InChannelIdx))
 		{
-			ConstantQResult->SortChronologically();
+			UE_LOG(LogAudioSynesthesia, Warning, TEXT("ConstantQNRT does not contain channel %d"), InChannelIdx);
+			return;
 		}
 
 		const TArray<Audio::FConstantQFrame>& ConstantQArray = ConstantQResult->GetFramesForChannel(InChannelIdx);
@@ -207,7 +217,9 @@ void UConstantQNRT::GetNormalizedChannelConstantQAtTime(const float InSeconds, c
 		// Normalize output by subtracting minimum and dividing by range
 		Audio::ArraySubtractByConstantInPlace(OutConstantQ, Settings->NoiseFloorDb);
 
-		float ConstantQRange = ConstantQResult->GetChannelConstantQRange(InChannelIdx, Settings->NoiseFloorDb);
+		FFloatInterval ConstantQInterval = ConstantQResult->GetChannelConstantQInterval(InChannelIdx);
+
+		const float ConstantQRange = ConstantQInterval.Max - FMath::Max(Settings->NoiseFloorDb, ConstantQInterval.Min);
 
 		if (ConstantQRange > SMALL_NUMBER)
 		{
@@ -219,13 +231,13 @@ void UConstantQNRT::GetNormalizedChannelConstantQAtTime(const float InSeconds, c
 	}
 }
 
-TUniquePtr<Audio::IAnalyzerNRTSettings> UConstantQNRT::GetSettings()
+TUniquePtr<Audio::IAnalyzerNRTSettings> UConstantQNRT::GetSettings(const float InSampleRate, const int32 InNumChannels) const
 {
 	TUniquePtr<Audio::IAnalyzerNRTSettings> AnalyzerSettings;
 
 	if (Settings)
 	{
-		AnalyzerSettings = Settings->GetSettings();	
+		AnalyzerSettings = Settings->GetSettings(InSampleRate, InNumChannels);	
 	}
 
 	return AnalyzerSettings;
diff --git a/Engine/Plugins/Runtime/AudioSynesthesia/Source/AudioSynesthesia/Private/ConstantQNRTFactory.cpp b/Engine/Plugins/Runtime/AudioSynesthesia/Source/AudioSynesthesia/Private/ConstantQNRTFactory.cpp
--- a/Engine/Plugins/Runtime/AudioSynesthesia/Source/AudioSynesthesia/Private/ConstantQNRTFactory.cpp
+++ b/Engine/Plugins/Runtime/AudioSynesthesia/Source/AudioSynesthesia/Private/ConstantQNRTFactory.cpp
@@ -1,6 +1,7 @@
 // Copyright 1998-2019 Epic Games, Inc. All Rights Reserved.
 
 #include "ConstantQNRTFactory.h"
+#include "AudioSynesthesiaCustomVersion.h"
 #include "DSP/ConstantQ.h"
 #include "DSP/DeinterleaveView.h"
 #include "DSP/FloatArrayMath.h"
@@ -13,6 +14,8 @@ namespace Audio
 
 	FArchive &operator <<(FArchive& Ar, FConstantQFrame& Frame)
 	{
+		Ar.UsingCustomVersion(FAudioSynesthesiaCustomVersion::GUID);
+
 		Ar << Frame.Channel;
 		Ar << Frame.Timestamp;
 		Ar << Frame.Spectrum;
@@ -21,11 +24,15 @@ namespace Audio
 	}
 
 	FConstantQNRTResult::FConstantQNRTResult()
-	:	bIsSortedChronologically(false)
+	:	DurationInSeconds(0.f)
+	,	bIsSortedChronologically(false)
 	{}
 
 	void FConstantQNRTResult::Serialize(FArchive& Archive)
 	{
+		Archive.UsingCustomVersion(FAudioSynesthesiaCustomVersion::GUID);
+		
+		Archive << DurationInSeconds;
 		Archive << bIsSortedChronologically;
 		Archive << ChannelCQTFrames;
 		Archive << ChannelCQTIntervals;
@@ -45,30 +52,34 @@ namespace Audio
 		bIsSortedChronologically = false;
 	}
 
+	bool FConstantQNRTResult::ContainsChannel(int32 InChannelIndex) const
+	{
+		return ChannelCQTFrames.Contains(InChannelIndex);
+	}
+
 	const TArray<FConstantQFrame>& FConstantQNRTResult::GetFramesForChannel(int32 InChannelIndex) const
 	{
 		return ChannelCQTFrames[InChannelIndex];
 	}
 
-	float FConstantQNRTResult::GetChannelConstantQRange(int32 InChannelIdx, float InNoiseFloor) const
+	FFloatInterval FConstantQNRTResult::GetChannelConstantQInterval(int32 InChannelIdx) const
 	{
-		const FFloatInterval& Interval = ChannelCQTIntervals.FindRef(InChannelIdx);
+		return ChannelCQTIntervals[InChannelIdx];
+	}
 
-		if (Interval.Contains(InNoiseFloor))
-		{
-			return Interval.Max - InNoiseFloor;
-		}
-		else if (InNoiseFloor > Interval.Max)
-		{
-			return 0.f;
-		}
+	void FConstantQNRTResult::GetChannels(TArray<int32>& OutChannels) const
+	{
+		ChannelCQTFrames.GetKeys(OutChannels);
+	}
 
-		return Interval.Size();
+	float FConstantQNRTResult::GetDurationInSeconds() const 
+	{
+		return DurationInSeconds;
 	}
 
-	int32 FConstantQNRTResult::GetNumChannels() const
+	void FConstantQNRTResult::SetDurationInSeconds(float InDuration)
 	{
-		return ChannelCQTFrames.Num();
+		DurationInSeconds = InDuration;
 	}
 
 	 /* Returns true if FConstantQNRTFrame arrays are sorted in chronologically ascending order via their timestamp.
@@ -96,7 +107,8 @@ namespace Audio
 	/*******************************************************************************/
 
 	FConstantQNRTWorker::FConstantQNRTWorker(const FAnalyzerNRTParameters& InParams, const FConstantQNRTSettings& InAnalyzerSettings)
-	:	NumChannels(InParams.NumChannels)
+	:	NumFrames(0)
+	,	NumChannels(InParams.NumChannels)
 	,	NumBuffers(0)
 	,	SampleRate(InParams.SampleRate)
 	,	NumHopFrames(0)
@@ -139,6 +151,8 @@ namespace Audio
 
 	void FConstantQNRTWorker::AnalyzeMultichannel(TArrayView<const float> InAudio, IAnalyzerNRTResult* OutResult, bool bDoFlush)
 	{
+		NumFrames += InAudio.Num() / NumChannels;
+
 		// Assume that outer layers ensured that this is of correct type.
 		FConstantQNRTResult* ConstantQResult = static_cast<FConstantQNRTResult*>(OutResult);
 
@@ -196,12 +210,23 @@ namespace Audio
 
 		check(nullptr != ConstantQResult);
 
-		NumBuffers = 0; // Reset internal counters
 
 		if (nullptr != OutResult)
 		{
+			check(SampleRate > 0.f);
+
+			float DurationInSeconds = static_cast<float>(NumFrames) / SampleRate;
+
+			ConstantQResult->SetDurationInSeconds(DurationInSeconds);
+
 			ConstantQResult->SortChronologically();
 		}
+
+		// Reset internal counters
+		NumBuffers = 0; 
+		NumFrames = 0;
+
+		SlidingBuffer.Reset();
 	}
 
 	void FConstantQNRTWorker::AnalyzeWindow(const AlignedFloatBuffer& InWindow, int32 InChannelIndex, FConstantQNRTResult& OutResult)
@@ -229,13 +254,13 @@ namespace Audio
 	}
 
 	// Create a new FConstantQNRTResult.
-	TUniquePtr<IAnalyzerNRTResult> FConstantQNRTFactory::NewResult()
+	TUniquePtr<IAnalyzerNRTResult> FConstantQNRTFactory::NewResult() const
 	{
 		return MakeUnique<FConstantQNRTResult>();
 	}
 
 	// Create a new FConstantQNRTWorker
-	TUniquePtr<IAnalyzerNRTWorker> FConstantQNRTFactory::NewWorker(const FAnalyzerNRTParameters& InParams, const IAnalyzerNRTSettings* InSettings)
+	TUniquePtr<IAnalyzerNRTWorker> FConstantQNRTFactory::NewWorker(const FAnalyzerNRTParameters& InParams, const IAnalyzerNRTSettings* InSettings) const
 	{
 		const FConstantQNRTSettings* ConstantQSettings = static_cast<const FConstantQNRTSettings*>(InSettings);
 
diff --git a/Engine/Plugins/Runtime/AudioSynesthesia/Source/AudioSynesthesia/Private/LoudnessNRT.cpp b/Engine/Plugins/Runtime/AudioSynesthesia/Source/AudioSynesthesia/Private/LoudnessNRT.cpp
--- a/Engine/Plugins/Runtime/AudioSynesthesia/Source/AudioSynesthesia/Private/LoudnessNRT.cpp
+++ b/Engine/Plugins/Runtime/AudioSynesthesia/Source/AudioSynesthesia/Private/LoudnessNRT.cpp
@@ -3,6 +3,7 @@
 #include "LoudnessNRT.h"
 #include "LoudnessNRTFactory.h"
 #include "InterpolateSorted.h"
+#include "AudioSynesthesiaLog.h"
 
 namespace
 {
@@ -36,7 +37,7 @@ ULoudnessNRTSettings::ULoudnessNRTSettings()
 ,	NoiseFloorDb(-60.f)
 {}
 
-TUniquePtr<Audio::IAnalyzerNRTSettings> ULoudnessNRTSettings::GetSettings()
+TUniquePtr<Audio::IAnalyzerNRTSettings> ULoudnessNRTSettings::GetSettings(const float InSampleRate, const int32 InNumChannels) const
 {
 	TUniquePtr<Audio::FLoudnessNRTSettings> Settings = MakeUnique<Audio::FLoudnessNRTSettings>();
 
@@ -80,26 +81,30 @@ ULoudnessNRT::ULoudnessNRT()
 	Settings = CreateDefaultSubobject<ULoudnessNRTSettings>(TEXT("DefaultLoudnessNRTSettings"));
 #if WITH_EDITOR
 	// Bind settings to audio analyze so changes to default settings will trigger analysis.
-	Settings->AnalyzeAudioDelegate.BindUObject(this, &UAudioAnalyzerNRT::AnalyzeAudio);
+	SetSettingsDelegate(Settings);
 #endif
 }
 
-void ULoudnessNRT::GetLoudnessAtTime(const float InSeconds, float& OutLoudness)
+void ULoudnessNRT::GetLoudnessAtTime(const float InSeconds, float& OutLoudness) const
 {
 	GetChannelLoudnessAtTime(InSeconds, Audio::FLoudnessNRTResult::ChannelIndexOverall, OutLoudness);
 }
 
-void ULoudnessNRT::GetChannelLoudnessAtTime(const float InSeconds, const int32 InChannel, float& OutLoudness)
+void ULoudnessNRT::GetChannelLoudnessAtTime(const float InSeconds, const int32 InChannel, float& OutLoudness) const
 {
 	OutLoudness = 0.0f;
 
-	TSharedPtr<Audio::FLoudnessNRTResult, ESPMode::ThreadSafe> LoudnessResult = GetResult<Audio::FLoudnessNRTResult>();
+	TSharedPtr<const Audio::FLoudnessNRTResult, ESPMode::ThreadSafe> LoudnessResult = GetResult<Audio::FLoudnessNRTResult>();
 
 	if (LoudnessResult.IsValid())
 	{
-		if (!LoudnessResult->IsSortedChronologically())
+		// The loudness result should never used here if it is not already sorted.
+		check(LoudnessResult->IsSortedChronologically());
+
+		if (!LoudnessResult->ContainsChannel(InChannel))
 		{
-			LoudnessResult->SortChronologically();
+			UE_LOG(LogAudioSynesthesia, Warning, TEXT("LoudnessNRT does not contain channel %d"), InChannel);
+			return;
 		}
 
 		const TArray<Audio::FLoudnessDatum>& LoudnessArray = LoudnessResult->GetChannelLoudnessArray(InChannel);
@@ -108,22 +113,26 @@ void ULoudnessNRT::GetChannelLoudnessAtTime(const float InSeconds, const int32 I
 	}
 }
 
-void ULoudnessNRT::GetNormalizedLoudnessAtTime(const float InSeconds, float& OutLoudness)
+void ULoudnessNRT::GetNormalizedLoudnessAtTime(const float InSeconds, float& OutLoudness) const
 {
 	GetNormalizedChannelLoudnessAtTime(InSeconds, Audio::FLoudnessNRTResult::ChannelIndexOverall, OutLoudness);
 }
 
-void ULoudnessNRT::GetNormalizedChannelLoudnessAtTime(const float InSeconds, const int32 InChannel, float& OutLoudness)
+void ULoudnessNRT::GetNormalizedChannelLoudnessAtTime(const float InSeconds, const int32 InChannel, float& OutLoudness) const
 {
 	OutLoudness = 0.0f;
 
-	TSharedPtr<Audio::FLoudnessNRTResult, ESPMode::ThreadSafe> LoudnessResult = GetResult<Audio::FLoudnessNRTResult>();
+	TSharedPtr<const Audio::FLoudnessNRTResult, ESPMode::ThreadSafe> LoudnessResult = GetResult<Audio::FLoudnessNRTResult>();
 
 	if (LoudnessResult.IsValid())
 	{
-		if (!LoudnessResult->IsSortedChronologically())
+		// The loudness result should never used here if it is not already sorted.
+		check(LoudnessResult->IsSortedChronologically());
+
+		if (!LoudnessResult->ContainsChannel(InChannel))
 		{
-			LoudnessResult->SortChronologically();
+			UE_LOG(LogAudioSynesthesia, Warning, TEXT("LoudnessNRT does not contain channel %d"), InChannel);
+			return;
 		}
 
 		const TArray<Audio::FLoudnessDatum>& LoudnessArray = LoudnessResult->GetChannelLoudnessArray(InChannel);
@@ -141,13 +150,13 @@ void ULoudnessNRT::GetNormalizedChannelLoudnessAtTime(const float InSeconds, con
 	}
 }
 
-TUniquePtr<Audio::IAnalyzerNRTSettings> ULoudnessNRT::GetSettings()
+TUniquePtr<Audio::IAnalyzerNRTSettings> ULoudnessNRT::GetSettings(const float InSampleRate, const int32 InNumChannels) const
 {
 	TUniquePtr<Audio::IAnalyzerNRTSettings> AnalyzerSettings;
 
 	if (Settings)
 	{
-		AnalyzerSettings = Settings->GetSettings();	
+		AnalyzerSettings = Settings->GetSettings(InSampleRate, InNumChannels);	
 	}
 
 	return AnalyzerSettings;
diff --git a/Engine/Plugins/Runtime/AudioSynesthesia/Source/AudioSynesthesia/Private/LoudnessNRTFactory.cpp b/Engine/Plugins/Runtime/AudioSynesthesia/Source/AudioSynesthesia/Private/LoudnessNRTFactory.cpp
--- a/Engine/Plugins/Runtime/AudioSynesthesia/Source/AudioSynesthesia/Private/LoudnessNRTFactory.cpp
+++ b/Engine/Plugins/Runtime/AudioSynesthesia/Source/AudioSynesthesia/Private/LoudnessNRTFactory.cpp
@@ -1,6 +1,7 @@
 // Copyright 1998-2019 Epic Games, Inc. All Rights Reserved.
 
 #include "LoudnessNRTFactory.h"
+#include "AudioSynesthesiaCustomVersion.h"
 #include "DSP/SlidingWindow.h"
 
 namespace Audio
@@ -10,6 +11,8 @@ namespace Audio
 	/************************************************************************/
 	FArchive &operator <<(FArchive& Ar, FLoudnessDatum& Datum)
 	{
+		Ar.UsingCustomVersion(FAudioSynesthesiaCustomVersion::GUID);
+
 		Ar << Datum.Channel;
 		Ar << Datum.Timestamp;
 		Ar << Datum.Energy;
@@ -23,16 +26,17 @@ namespace Audio
 
 
 	FLoudnessNRTResult::FLoudnessNRTResult()
-	: bIsSortedChronologically(false)
+	:	DurationInSeconds(0.f)
+	,	bIsSortedChronologically(false)
 	{
-		// By  default, overall loudness exists
-		ChannelLoudnessArrays.Add(ChannelIndexOverall);
-		ChannelLoudnessIntervals.Add(ChannelIndexOverall);
 	}
 
 
 	void FLoudnessNRTResult::Serialize(FArchive& Archive)
 	{
+		Archive.UsingCustomVersion(FAudioSynesthesiaCustomVersion::GUID);
+
+		Archive << DurationInSeconds;
 		Archive << bIsSortedChronologically;
 		Archive << ChannelLoudnessArrays;
 		Archive << ChannelLoudnessIntervals;
@@ -52,6 +56,11 @@ namespace Audio
 		bIsSortedChronologically = false;
 	}
 
+	bool FLoudnessNRTResult::ContainsChannel(int32 InChannelIndex) const
+	{
+		return ChannelLoudnessArrays.Contains(InChannelIndex);
+	}
+
 	const TArray<FLoudnessDatum>& FLoudnessNRTResult::GetChannelLoudnessArray(int32 ChannelIdx) const
 	{
 		return ChannelLoudnessArrays[ChannelIdx];
@@ -85,9 +94,19 @@ namespace Audio
 		return Interval.Size();
 	}
 
-	int32 FLoudnessNRTResult::GetNumChannels() const 
+	void FLoudnessNRTResult::GetChannels(TArray<int32>& OutChannels) const
 	{
-		return FMath::Max(0, ChannelLoudnessArrays.Num() - 1);
+		ChannelLoudnessArrays.GetKeys(OutChannels);
+	}
+
+	float FLoudnessNRTResult::GetDurationInSeconds() const 
+	{
+		return DurationInSeconds;
+	}
+
+	void FLoudnessNRTResult::SetDurationInSeconds(float InDuration)
+	{
+		DurationInSeconds = InDuration;
 	}
 
 	bool FLoudnessNRTResult::IsSortedChronologically() const
@@ -113,9 +132,11 @@ namespace Audio
 	: NumChannels(InParams.NumChannels)
 	, NumAnalyzedBuffers(0)
 	, NumHopFrames(0)
+	, NumFrames(0)
 	, SampleRate(InParams.SampleRate)
 	{
-		check(NumChannels >= 0);
+		check(NumChannels > 0);
+		check(SampleRate > 0.f);
 
 		// Include NumChannels in calculating window size because windows are generated 
 		// with interleaved samples and deinterleaved later. 
@@ -134,6 +155,8 @@ namespace Audio
 
 	void FLoudnessNRTWorker::Analyze(TArrayView<const float> InAudio, IAnalyzerNRTResult* OutResult) 
 	{
+		NumFrames += InAudio.Num() / NumChannels;
+
 		// Assume that outer layers ensured that this is of correct type.
 		FLoudnessNRTResult* LoudnessResult = static_cast<FLoudnessNRTResult*>(OutResult);
 
@@ -168,8 +191,11 @@ namespace Audio
 		// Prepare result for fast lookup.
 		LoudnessResult->SortChronologically(); 
 
+		LoudnessResult->SetDurationInSeconds(static_cast<float>(NumFrames) / SampleRate);
+
 		// Reset internal counters
 		NumAnalyzedBuffers = 0; 
+		NumFrames = 0;
 
 		// Reset sliding buffer
 		InternalBuffer->Reset();
@@ -231,13 +257,13 @@ namespace Audio
 		return TEXT("Loudness Analyzer Non-Real-Time");
 	}
 
-	TUniquePtr<IAnalyzerNRTResult> FLoudnessNRTFactory::NewResult()
+	TUniquePtr<IAnalyzerNRTResult> FLoudnessNRTFactory::NewResult() const
 	{
 		TUniquePtr<FLoudnessNRTResult> Result = MakeUnique<FLoudnessNRTResult>();
 		return Result;
 	}
 
-	TUniquePtr<IAnalyzerNRTWorker> FLoudnessNRTFactory::NewWorker(const FAnalyzerNRTParameters& InParams, const IAnalyzerNRTSettings* InSettings)
+	TUniquePtr<IAnalyzerNRTWorker> FLoudnessNRTFactory::NewWorker(const FAnalyzerNRTParameters& InParams, const IAnalyzerNRTSettings* InSettings) const
 	{
 		const FLoudnessNRTSettings* LoudnessSettings = static_cast<const FLoudnessNRTSettings*>(InSettings);
 
diff --git a/Engine/Plugins/Runtime/AudioSynesthesia/Source/AudioSynesthesia/Private/OnsetAnalyzer.cpp b/Engine/Plugins/Runtime/AudioSynesthesia/Source/AudioSynesthesia/Private/OnsetAnalyzer.cpp
new file mode 100644
--- /dev/null
+++ b/Engine/Plugins/Runtime/AudioSynesthesia/Source/AudioSynesthesia/Private/OnsetAnalyzer.cpp
@@ -0,0 +1,236 @@
+// Copyright 1998-2019 Epic Games, Inc. All Rights Reserved.
+
+#include "OnsetAnalyzer.h"
+#include "PeakPicker.h"
+#include "DSP/FloatArrayMath.h"
+
+namespace Audio
+{
+	FOnsetStrengthAnalyzer::FOnsetStrengthAnalyzer(const FOnsetStrengthSettings& InSettings, float InSampleRate)
+	:	Settings(InSettings)
+	,	LagSpectraIndex(0)
+	,	SlidingBuffer(InSettings.NumWindowFrames, InSettings.NumHopFrames)
+	,	Window(Settings.WindowType, Settings.NumWindowFrames, 1, false)
+	{
+		// Check values legit.
+		check(InSampleRate > 0.f);
+		check(Settings.FFTSize > 0);
+		check(Settings.NumHopFrames > 0);
+		check(Settings.NumWindowFrames > 0);
+		check(Settings.ComparisonLag > 0);
+		check(Settings.MelSettings.NumBands > 0);
+		check(Settings.NumWindowFrames <= Settings.FFTSize);
+
+		// intialize buffers dealing with audio window
+		WorkingBuffer.AddUninitialized(Settings.NumWindowFrames);
+
+		// initialize buffers dealing with fft input / output
+		WindowedSamples.AddZeroed(Settings.FFTSize);
+		FFTOutputRealData.AddUninitialized(Settings.FFTSize);
+		FFTOutputImagData.AddUninitialized(Settings.FFTSize);
+		FFTSpectrumBuffer.AddUninitialized(Settings.FFTSize);
+		
+		// intialize buffers dealing with perceptual spectrum 
+		PreviousSpectra.AddDefaulted(Settings.ComparisonLag);
+		MelSpectrum.AddUninitialized(Settings.MelSettings.NumBands);
+		SpectrumDifference.AddUninitialized(Settings.MelSettings.NumBands);
+
+		// create mel spectrum transform
+		MelTransform = NewMelSpectrumKernelTransform(Settings.MelSettings, Settings.FFTSize, InSampleRate);
+
+		check(MelTransform.IsValid());
+	}
+
+	void FOnsetStrengthAnalyzer::CalculateOnsetStrengths(TArrayView<const float> InSamples, TArray<float>& OutEnvelopeStrengths)
+	{
+		const bool bDoFlush = false;
+		AnalyzeAudio(InSamples, OutEnvelopeStrengths, bDoFlush);
+	}
+
+	void FOnsetStrengthAnalyzer::FlushAudio(TArray<float>& OutEnvelopeStrengths)
+	{
+		TArray<float> EmptyArray;
+
+		// Force flushing of internal audio buffer
+		const bool bDoFlush = true;
+
+		AnalyzeAudio(EmptyArray, OutEnvelopeStrengths, bDoFlush);
+	}
+	
+
+	void FOnsetStrengthAnalyzer::Reset()
+	{
+		// Reset internal memory
+		for (TArray<float>& Spectrum : PreviousSpectra)
+		{
+			Spectrum.Reset();
+		}
+		SlidingBuffer.Reset();
+	}
+
+	float FOnsetStrengthAnalyzer::GetTimestampForIndex(const FOnsetStrengthSettings& InSettings, float InSampleRate, int32 InIndex)
+	{
+		check(InSampleRate > 0.f);
+
+		// Getting the exact timestamp of an onset from the onset strength envelope is non-trivial because
+		// the onset strength envelope changes depending on the attack, decay and sustain of the onset as well
+		// as the choice of window, window size, etc. etc. 
+		//
+		// This approach does not attempt to account for all the subtle variability than can come up, but 
+		// uses a generally sane approach of estimating the timestamp to be somewhere between the center of
+		// the "loud" window and the end of the "quiet" window.
+		const int32 WindowOffset = InSettings.NumWindowFrames / 2;
+
+		const int32 LoudFrameIndex = (InIndex * InSettings.NumHopFrames) + WindowOffset;
+
+		// Get the end of the window to which this was compared.
+		const int32 QuietIndex = InIndex - InSettings.ComparisonLag;
+		const int32 QuietFrameIndex = (QuietIndex * InSettings.NumHopFrames) + InSettings.NumWindowFrames;
+
+		const int32 EstimatedOnsetFrameIndex = (LoudFrameIndex + QuietFrameIndex) / 2;
+
+		return static_cast<float>(EstimatedOnsetFrameIndex) / InSampleRate;
+	}
+
+	void FOnsetStrengthAnalyzer::AnalyzeAudio(TArrayView<const float> InSamples, TArray<float>& OutEnvelopeStrengths, bool bDoFlush)
+	{
+		typedef TAutoSlidingWindow<float, FAudioBufferAlignedAllocator> FSlidingWindow;
+
+		// Going to fill OutEnvelopeStrengths with envelope strength values generated during this call.
+		OutEnvelopeStrengths.Reset();
+
+		// Slide over audio 
+		FSlidingWindow SlidingWindow(SlidingBuffer, InSamples, WorkingBuffer, bDoFlush);
+		for (const AlignedFloatBuffer& WindowBuffer : SlidingWindow)
+		{
+			// Calculate onset strength per audio buffer.
+			float OnsetStrength = GetNextOnsetStrength(WindowBuffer);
+
+			OutEnvelopeStrengths.Add(OnsetStrength);
+		}
+	}
+
+	float FOnsetStrengthAnalyzer::GetNextOnsetStrength(TArrayView<const float> InSamples)
+	{
+		float OnsetStrength = 0.f;
+
+		// Copy input samples and apply window
+		FMemory::Memcpy(WindowedSamples.GetData(), InSamples.GetData(), sizeof(float) * Settings.NumWindowFrames);
+		Window.ApplyToBuffer(WindowedSamples.GetData());
+
+		// Take FFT of windowed samples
+		const FFTTimeDomainData TimeData = {WindowedSamples.GetData(), Settings.FFTSize};
+		FFTFreqDomainData FreqData = {FFTOutputRealData.GetData(), FFTOutputImagData.GetData()};
+
+		PerformFFT(TimeData, FreqData);
+
+		// Calculate spectrum from FFT output
+		ComputePowerSpectrum(FreqData, Settings.FFTSize, FFTSpectrumBuffer);
+		
+		// Convert spectrum to CQT
+		MelTransform->TransformArray(FFTSpectrumBuffer, MelSpectrum);
+
+		// Apply decibel scaling 
+		ArrayPowerToDecibelInPlace(MelSpectrum);
+
+		// Clamp to noise floor
+		ArrayClampMinInPlace(MelSpectrum, Settings.NoiseFloorDb);
+
+		// Onset strength is mean(abs(diff(spectrum - previous_spectrum)))
+		TArray<float>& PreviousSpectrum = PreviousSpectra[LagSpectraIndex];
+		if (PreviousSpectrum.Num() > 0)
+		{
+			ArraySubtract(MelSpectrum, PreviousSpectrum, SpectrumDifference);
+
+			// Half wave rectify
+			ArrayClampMinInPlace(SpectrumDifference, 0.f);
+
+			ArrayMean(SpectrumDifference, OnsetStrength);
+		}
+
+		// Save this spectrum for later use
+		PreviousSpectrum.Reset(Settings.MelSettings.NumBands);
+		PreviousSpectrum.AddUninitialized(Settings.MelSettings.NumBands);
+		FMemory::Memcpy(PreviousSpectrum.GetData(), MelSpectrum.GetData(), sizeof(float) * Settings.MelSettings.NumBands);
+
+		// Increment index for comparison spectra
+		LagSpectraIndex = (LagSpectraIndex + 1) % Settings.ComparisonLag;
+
+		return OnsetStrength;
+	}
+
+
+	/*******************************************************************************/
+	/*******************************************************************************/
+	/*******************************************************************************/
+
+	void OnsetExtractIndices(const FPeakPickerSettings& InSettings, TArrayView<const float> InOnsetEnvelope, TArray<int32>& OutOnsetIndices)
+	{
+		const int32 Num = InOnsetEnvelope.Num();
+		OutOnsetIndices.Reset();
+
+		if (Num < 1)
+		{
+			return;
+		}
+		
+		TArray<float> NormOnsets;
+
+		// Onset strength is normalized between [0, 1.0] before picking peaks.
+		ArrayMinMaxNormalize(InOnsetEnvelope, NormOnsets);
+
+		FPeakPicker PeakPicker(InSettings);
+		PeakPicker.PickPeaks(NormOnsets, OutOnsetIndices);
+	}
+
+	/*******************************************************************************/
+	/*******************************************************************************/
+	/*******************************************************************************/
+
+	void OnsetBacktrackIndices(TArrayView<const float> InOnsetStrength, TArrayView<const int32> InOnsetIndices, TArray<int32>& OutBacktrackedOnsetIndices)
+	{
+		const int32 NumOnsets = InOnsetIndices.Num();
+		const int32 NumOnsetStrength = InOnsetStrength.Num();
+
+		OutBacktrackedOnsetIndices.Reset();
+		
+		if (NumOnsets < 1)
+		{
+			return;
+		}
+
+		OutBacktrackedOnsetIndices.AddUninitialized(NumOnsets);
+
+		const int32* InOnsetIndicesData = InOnsetIndices.GetData();
+		int32* OutOnsetIndicesData = OutBacktrackedOnsetIndices.GetData();
+		const float* InOnsetStrengthData = InOnsetStrength.GetData();
+
+		for (int32 i = 0; i < NumOnsets; i++)
+		{
+
+			int32 OnsetStartIndex = InOnsetIndicesData[i];
+			check(OnsetStartIndex < NumOnsetStrength);
+
+			while (OnsetStartIndex > 0)
+			{
+				if (InOnsetStrengthData[OnsetStartIndex] <= 0.f)
+				{
+					// No onset strength means the audio is steady state or silent.
+					break;
+				}
+				else if (InOnsetStrengthData[OnsetStartIndex - 1] <= InOnsetStrengthData[OnsetStartIndex])
+				{
+					// The onest strength is still holding steady or decreasing.
+					OnsetStartIndex--;
+				}
+				else
+				{
+					// Onset strength is no longer decreasing.
+					break;
+				}
+			}
+
+			OutOnsetIndicesData[i] = OnsetStartIndex;
+		}
+	}
+}
diff --git a/Engine/Plugins/Runtime/AudioSynesthesia/Source/AudioSynesthesia/Private/OnsetNRT.cpp b/Engine/Plugins/Runtime/AudioSynesthesia/Source/AudioSynesthesia/Private/OnsetNRT.cpp
new file mode 100644
--- /dev/null
+++ b/Engine/Plugins/Runtime/AudioSynesthesia/Source/AudioSynesthesia/Private/OnsetNRT.cpp
@@ -0,0 +1,205 @@
+// Copyright 1998-2019 Epic Games, Inc. All Rights Reserved.
+
+#include "OnsetNRT.h"
+#include "OnsetNRTFactory.h"
+#include "AudioSynesthesiaLog.h"
+
+namespace
+{
+	// Looks up onsets that fall between two timestamps and adds there info to the output arrays. 
+	void AddOnsetsInRange(const Audio::FOnsetNRTResult& OnsetResult, const float InStartSeconds, const float InEndSeconds, const int32 InChannel, TArray<float>& OutOnsetTimestamps, TArray<float>& OutOnsetStrengths)
+	{
+		// Onset results should always be sorted temporally within UOnsetNRT objects
+		check(OnsetResult.IsSortedChronologically());
+
+		const TArray<Audio::FOnset>& OnsetArray = OnsetResult.GetOnsetsForChannel(InChannel);
+
+		// By design we are not inclusive of InStartSeconds but are inclusive with the InEndSeconds timestamp. 
+		// This is to avoid returning duplicate onsets when this function is called repeatedly with incremented
+		// timestamps.
+		int32 StartTimestampIndex = Algo::UpperBoundBy(OnsetArray, InStartSeconds, [](const Audio::FOnset& Onset) { return Onset.Timestamp; });
+
+		int32 EndTimestampIndex = Algo::UpperBoundBy(OnsetArray, InEndSeconds, [](const Audio::FOnset& Onset) { return Onset.Timestamp; });
+		
+
+		// Add to output arrays.
+		for (int32 i = StartTimestampIndex; i < EndTimestampIndex; i++)
+		{
+			const Audio::FOnset& Onset = OnsetArray[i];
+			OutOnsetTimestamps.Add(Onset.Timestamp);
+			OutOnsetStrengths.Add(Onset.Strength);
+		}
+	}
+}
+
+/*************************************************************************************/
+/***************************** UOnsetNRTSettings *************************************/
+/*************************************************************************************/
+
+// Construct default settings.
+UOnsetNRTSettings::UOnsetNRTSettings()
+:	bDownmixToMono(false)
+,	GranularityInSeconds(0.01f)
+,	Sensitivity(0.5f)
+,	MinimumFrequency(20.f)
+,	MaximumFrequency(16000.f)
+{}
+
+TUniquePtr<Audio::IAnalyzerNRTSettings> UOnsetNRTSettings::GetSettings(const float InSampleRate, const int32 InNumChannels) const
+{
+	// Make a default FOnsetNRTSettings object
+	TUniquePtr<Audio::FOnsetNRTSettings> OnsetNRTSettings = MakeUnique<Audio::FOnsetNRTSettings>();
+
+	// Apply global settings
+	OnsetNRTSettings->bDownmixToMono = bDownmixToMono;
+
+	// Apply Onset strength settings
+	// Note: Granularity effects hop size, window size and fft size.
+	OnsetNRTSettings->OnsetStrengthSettings.NumHopFrames = FMath::Max(FMath::RoundToInt(GranularityInSeconds * InSampleRate), 1);
+	OnsetNRTSettings->OnsetStrengthSettings.NumWindowFrames = 4 * OnsetNRTSettings->OnsetStrengthSettings.NumHopFrames;
+	OnsetNRTSettings->OnsetStrengthSettings.ComparisonLag = 4;
+	OnsetNRTSettings->OnsetStrengthSettings.FFTSize = FMath::RoundUpToPowerOfTwo(OnsetNRTSettings->OnsetStrengthSettings.NumWindowFrames);
+	OnsetNRTSettings->OnsetStrengthSettings.WindowType = Audio::EWindowType::Blackman;
+
+	OnsetNRTSettings->OnsetStrengthSettings.NoiseFloorDb = -60; 
+	// Onset strength mel band settings
+	OnsetNRTSettings->OnsetStrengthSettings.MelSettings.NumBands = 26;
+	OnsetNRTSettings->OnsetStrengthSettings.MelSettings.KernelMinCenterFreq = FMath::Max(MinimumFrequency * 1.1f, 20.f);
+	OnsetNRTSettings->OnsetStrengthSettings.MelSettings.KernelMaxCenterFreq = FMath::Min(MaximumFrequency * 0.9f, 20000.f);
+	OnsetNRTSettings->OnsetStrengthSettings.MelSettings.Normalization = Audio::EMelNormalization::EqualEnergy;
+	OnsetNRTSettings->OnsetStrengthSettings.MelSettings.BandWidthStretch = 1.f;
+
+	// Peak picker settings
+	// 
+	// The "Sensitivity" variable has a strong influence on the peak picker settings. 
+	// A higher sensitivity results in many more peaks being picked, and a lower sensitivity
+	// results in only the largest peaks being picked. 
+	
+	// Peak picker is picking peaks from analyzed audio windows. So we need to convert timestamps 
+	// to window indices.
+	const float HopRate = InSampleRate / static_cast<float>(OnsetNRTSettings->OnsetStrengthSettings.NumHopFrames);
+	TFunction<int32 (float)> TimeToWindow = [HopRate](float InTime) 
+	{
+		const int32 MinSample = 1;
+		
+		return FMath::Max(MinSample, FMath::RoundToInt(InTime * HopRate));
+	};
+
+	// Flip sensitivity when doing interpolation. An alpha == 1 is not sensitive. An alpha == 0 is most sensitive.
+	const float Alpha = 1.f - Sensitivity;
+
+	static const float MinPreMaxSeconds = 0.01f;
+	static const float MaxPreMaxSeconds = 0.3f;
+	OnsetNRTSettings->PeakPickerSettings.NumPreMax = TimeToWindow(FMath::InterpExpoIn(MinPreMaxSeconds, MaxPreMaxSeconds, Alpha));
+
+	static const float MinPostMaxSeconds = 0.01f;
+	static const float MaxPostMaxSeconds = 0.1f;
+	OnsetNRTSettings->PeakPickerSettings.NumPostMax = TimeToWindow(FMath::InterpExpoIn(MinPostMaxSeconds, MaxPostMaxSeconds, Alpha));
+
+	static const float MinPreMeanSeconds = 0.05f;
+	static const float MaxPreMeanSeconds = 0.5f;
+	OnsetNRTSettings->PeakPickerSettings.NumPreMean = TimeToWindow(FMath::InterpExpoIn(MinPreMeanSeconds, MaxPreMeanSeconds, Alpha));
+
+	static const float MinPostMeanSeconds = 0.05f;
+	static const float MaxPostMeanSeconds = 0.5f;
+	OnsetNRTSettings->PeakPickerSettings.NumPostMean = TimeToWindow(FMath::InterpExpoIn(MinPostMeanSeconds, MaxPostMeanSeconds, Alpha));
+
+	static const float MinWaitSeconds = 0.01f;
+	static const float MaxWaitSeconds = 0.25f;
+	OnsetNRTSettings->PeakPickerSettings.NumWait = TimeToWindow(FMath::InterpExpoIn(MinWaitSeconds, MaxWaitSeconds, Alpha));
+
+	static const float MinMeanDelta = 0.01f;
+	static const float MaxMeanDelta = 0.2f;
+	OnsetNRTSettings->PeakPickerSettings.MeanDelta = FMath::InterpExpoIn(MinMeanDelta, MaxMeanDelta, Alpha);
+
+	return OnsetNRTSettings;
+
+}
+
+
+/*************************************************************************************/
+/********************************* UOnsetNRT *****************************************/
+/*************************************************************************************/
+
+UOnsetNRT::UOnsetNRT()
+{
+	Settings = CreateDefaultSubobject<UOnsetNRTSettings>(TEXT("DefaultOnsetNRTSettings"));
+
+#if WITH_EDITOR
+	// Bind settings to audio analyze so changes to default settings will trigger analysis.
+	SetSettingsDelegate(Settings);
+#endif
+}
+
+
+void UOnsetNRT::GetChannelOnsetsBetweenTimes(const float InStartSeconds, const float InEndSeconds, const int32 InChannel, TArray<float>& OutOnsetTimestamps, TArray<float>& OutOnsetStrengths) const
+{
+	// Clear output arrays
+	OutOnsetTimestamps.Reset();
+	OutOnsetStrengths.Reset();
+
+	TSharedPtr<const Audio::FOnsetNRTResult, ESPMode::ThreadSafe> OnsetResult = GetResult<Audio::FOnsetNRTResult>();
+
+	if (OnsetResult.IsValid())
+	{
+		if (!OnsetResult->ContainsChannel(InChannel))
+		{
+			// invalid channel index.
+			UE_LOG(LogAudioSynesthesia, Warning, TEXT("OnsetNRT does not contain channel %d"), InChannel);
+			return;
+		}
+
+		AddOnsetsInRange(*OnsetResult, InStartSeconds, InEndSeconds, InChannel, OutOnsetTimestamps, OutOnsetStrengths);
+	}
+}
+
+void UOnsetNRT::GetNormalizedChannelOnsetsBetweenTimes(const float InStartSeconds, const float InEndSeconds, const int32 InChannel, TArray<float>& OutOnsetTimestamps, TArray<float>& OutOnsetStrengths) const
+{
+	// Clear output arrays
+	OutOnsetTimestamps.Reset();
+	OutOnsetStrengths.Reset();
+
+	TSharedPtr<const Audio::FOnsetNRTResult, ESPMode::ThreadSafe> OnsetResult = GetResult<Audio::FOnsetNRTResult>();
+
+	if (OnsetResult.IsValid())
+	{
+		if (!OnsetResult->ContainsChannel(InChannel))
+		{
+			// Invalid channel index.
+			UE_LOG(LogAudioSynesthesia, Warning, TEXT("OnsetNRT does not contain channel %d"), InChannel);
+			return;
+		}
+
+		AddOnsetsInRange(*OnsetResult, InStartSeconds, InEndSeconds, InChannel, OutOnsetTimestamps, OutOnsetStrengths);
+
+		// Normalize by maximum onset. Minimum onset value is assumed to be zero.
+		const FFloatInterval OnsetStrengthInterval = OnsetResult->GetChannelOnsetInterval(InChannel);
+		
+		if (OnsetStrengthInterval.Max > 0.f)
+		{
+			for (float& Strength : OutOnsetStrengths)
+			{
+				Strength /= OnsetStrengthInterval.Max;
+			}
+		}
+	}
+}
+
+TUniquePtr<Audio::IAnalyzerNRTSettings> UOnsetNRT::GetSettings(const float InSampleRate, const int32 InNumChannels) const
+{
+	TUniquePtr<Audio::IAnalyzerNRTSettings> AnalyzerSettings;
+
+	if (Settings)
+	{
+		AnalyzerSettings = Settings->GetSettings(InSampleRate, InNumChannels);	
+	}
+
+	return AnalyzerSettings;
+}
+
+FName UOnsetNRT::GetAnalyzerNRTFactoryName() const
+{
+	static const FName FactoryName(TEXT("OnsetNRTFactory"));
+	return FactoryName;
+}
+
diff --git a/Engine/Plugins/Runtime/AudioSynesthesia/Source/AudioSynesthesia/Private/OnsetNRTFactory.cpp b/Engine/Plugins/Runtime/AudioSynesthesia/Source/AudioSynesthesia/Private/OnsetNRTFactory.cpp
new file mode 100644
--- /dev/null
+++ b/Engine/Plugins/Runtime/AudioSynesthesia/Source/AudioSynesthesia/Private/OnsetNRTFactory.cpp
@@ -0,0 +1,326 @@
+// Copyright 1998-2019 Epic Games, Inc. All Rights Reserved.
+
+#include "OnsetNRTFactory.h"
+#include "OnsetAnalyzer.h"
+#include "AudioSynesthesiaCustomVersion.h"
+#include "DSP/DeinterleaveView.h"
+#include "DSP/FloatArrayMath.h"
+
+namespace Audio 
+{
+	/*******************************************************************************/
+	/*************************** FOnsetNRTResult *******************************/
+	/*******************************************************************************/
+
+	FArchive &operator <<(FArchive& Ar, FOnset& Frame)
+	{
+		Ar.UsingCustomVersion(FAudioSynesthesiaCustomVersion::GUID);
+
+		Ar << Frame.Channel;
+		Ar << Frame.Timestamp;
+		Ar << Frame.Strength;
+
+		return Ar;
+	}
+
+	FOnsetNRTResult::FOnsetNRTResult()
+	:	DurationInSeconds(0.f)
+	,	bIsSortedChronologically(false)
+	{}
+
+	void FOnsetNRTResult::Serialize(FArchive& Archive)
+	{
+		Archive.UsingCustomVersion(FAudioSynesthesiaCustomVersion::GUID);
+
+		Archive << DurationInSeconds;
+		Archive << bIsSortedChronologically;
+		Archive << ChannelOnsets;
+		Archive << ChannelOnsetStrengthIntervals;
+	}
+
+	void FOnsetNRTResult::AddOnset(int32 InChannelIndex, float InTimestamp, float InStrength)
+	{
+		// Add to onset arrays
+		TArray<FOnset>& OnsetArray = ChannelOnsets.FindOrAdd(InChannelIndex);
+		OnsetArray.Emplace(InChannelIndex, InTimestamp, InStrength);
+
+		// Store min/max value
+		FFloatInterval& Interval = ChannelOnsetStrengthIntervals.FindOrAdd(InChannelIndex);
+		Interval.Include(InStrength);
+
+		bIsSortedChronologically = false;
+	}
+
+	void FOnsetNRTResult::AddChannel(int32 InChannelIndex)
+	{
+		// Add an empty channel. Useful when there is a channel with no onsets.
+		ChannelOnsets.FindOrAdd(InChannelIndex);
+		ChannelOnsetStrengthIntervals.FindOrAdd(InChannelIndex);
+	}
+
+	bool FOnsetNRTResult::ContainsChannel(int32 InChannelIndex) const
+	{
+		return ChannelOnsets.Contains(InChannelIndex);
+	}
+
+	const TArray<FOnset>& FOnsetNRTResult::GetOnsetsForChannel(int32 InChannelIndex) const
+	{
+		check(ChannelOnsets.Contains(InChannelIndex));
+
+		return ChannelOnsets[InChannelIndex];
+	}
+
+	FFloatInterval FOnsetNRTResult::GetChannelOnsetInterval(int32 InChannelIndex) const
+	{
+		check(ChannelOnsets.Contains(InChannelIndex));
+
+		return ChannelOnsetStrengthIntervals[InChannelIndex];
+	}
+
+	void FOnsetNRTResult::GetChannels(TArray<int32>& OutChannels) const
+	{
+		ChannelOnsets.GetKeys(OutChannels);
+	}
+
+	float FOnsetNRTResult::GetDurationInSeconds() const 
+	{
+		return DurationInSeconds;
+	}
+
+	void FOnsetNRTResult::SetDurationInSeconds(float InDuration)
+	{
+		DurationInSeconds = InDuration;
+	}
+
+	 /* Returns true if FOnsetNRTFrame arrays are sorted in chronologically ascending order via their timestamp.
+	 */
+	bool FOnsetNRTResult::IsSortedChronologically() const
+	{
+		return bIsSortedChronologically;
+	}
+
+	/**
+	 * Sorts FOnsetNRTFrame arrays in chronologically ascnding order via their timestamp.
+	 */
+	void FOnsetNRTResult::SortChronologically()
+	{
+		for (auto& KeyValue : ChannelOnsets)
+		{
+			KeyValue.Value.Sort([](const FOnset& A, const FOnset& B) { return A.Timestamp < B.Timestamp; });
+		}
+
+		bIsSortedChronologically = true;
+	}
+
+	/*******************************************************************************/
+	/*************************** FOnsetNRTWorker *******************************/
+	/*******************************************************************************/
+
+	FOnsetNRTWorker::FOnsetNRTWorker(const FAnalyzerNRTParameters& InParams, const FOnsetNRTSettings& InAnalyzerSettings)
+	:	Settings(InAnalyzerSettings)
+	,	NumFrames(0)
+	,	NumChannels(InParams.NumChannels)
+	,	SampleRate(InParams.SampleRate)
+	,	NumWindowFrames(4096)
+	,	NumWindowSamples(0)
+	,	MonoScaling(1.f)
+	{
+		check(NumChannels > 0);
+		check(SampleRate > 0.f);
+
+		if (Settings.bDownmixToMono)
+		{
+			MonoBuffer.Reset(NumWindowFrames);
+			MonoBuffer.AddUninitialized(NumWindowFrames);
+			// equal power mixing
+			MonoScaling = 1.f / FMath::Sqrt(static_cast<float>(NumChannels));
+
+			// Only need one onset strength analyzer
+			TUniquePtr<FOnsetStrengthAnalyzer> Analyzer = MakeUnique<FOnsetStrengthAnalyzer>(InAnalyzerSettings.OnsetStrengthSettings, SampleRate);
+			OnsetStrengthAnalyzers.Emplace(MoveTemp(Analyzer));
+
+			// Only need one onset strength array
+			OnsetStrengths.AddDefaulted(1);
+		}
+		else
+		{
+			// Make individual onset strength analyzer for each channel
+			for (int32 i = 0; i < NumChannels; i++)
+			{
+				TUniquePtr<FOnsetStrengthAnalyzer> Analyzer = MakeUnique<FOnsetStrengthAnalyzer>(InAnalyzerSettings.OnsetStrengthSettings, SampleRate);
+				OnsetStrengthAnalyzers.Emplace(MoveTemp(Analyzer));
+			}
+			OnsetStrengths.AddDefaulted(NumChannels);
+		}
+
+		NumWindowSamples = NumChannels * NumWindowFrames;
+		
+		// Sliding buffer used for deinterleaving audio
+		SlidingBuffer = MakeUnique<TSlidingBuffer<float> >(NumWindowSamples, NumWindowSamples);
+	}
+
+
+	void FOnsetNRTWorker::Analyze(TArrayView<const float> InAudio, IAnalyzerNRTResult* OutResult)
+	{
+		bool bDoFlush = false;
+
+		AnalyzeMultichannel(InAudio, OutResult, bDoFlush);
+	}
+
+	void FOnsetNRTWorker::AnalyzeMultichannel(TArrayView<const float> InAudio, IAnalyzerNRTResult* OutResult, bool bDoFlush)
+	{
+		NumFrames += InAudio.Num() / NumChannels;
+
+		// Assume that outer layers ensured that this is of correct type.
+		FOnsetNRTResult* OnsetResult = static_cast<FOnsetNRTResult*>(OutResult);
+
+		check(nullptr != OnsetResult);
+
+		TAutoSlidingWindow<float, FAudioBufferAlignedAllocator> SlidingWindow(*SlidingBuffer, InAudio, HopBuffer, bDoFlush);
+
+		if (Settings.bDownmixToMono)
+		{
+			for (const AlignedFloatBuffer& Window : SlidingWindow)
+			{
+				FMemory::Memset(MonoBuffer.GetData(), 0, sizeof(float) * NumWindowFrames);
+
+				// Sum all channels into mono buffer
+				TAutoDeinterleaveView<float, FAudioBufferAlignedAllocator> DeinterleaveView(Window, ChannelBuffer, NumChannels);
+				for(auto Channel : DeinterleaveView)
+				{
+					MixInBufferFast(Channel.Values, MonoBuffer);		
+				}
+
+				// Equal power sum. assuming incoherrent signals.
+				MultiplyBufferByConstantInPlace(MonoBuffer, 1.f / FMath::Sqrt(static_cast<float>(NumChannels)));
+
+				AnalyzeWindow(MonoBuffer, 0, *OnsetResult);
+			}
+		}
+		else
+		{
+			// Loop through entire array of input samples.
+			for (const AlignedFloatBuffer& Window : SlidingWindow)
+			{
+				// Each channel is analyzed seperately.
+				for(const TDeinterleaveView<float>::TChannel<FAudioBufferAlignedAllocator>& Channel : TAutoDeinterleaveView<float, FAudioBufferAlignedAllocator>(Window, ChannelBuffer, NumChannels))
+				{
+					AnalyzeWindow(Channel.Values, Channel.ChannelIndex, *OnsetResult);
+				}
+			}
+		}
+	}
+
+	/** Called when analysis of audio asset is complete. */
+	void FOnsetNRTWorker::Finalize(IAnalyzerNRTResult* OutResult) 
+	{
+		// Finish analyzing audio
+		AlignedFloatBuffer EmptyArray;
+		bool bDoFlush = true;
+
+		AnalyzeMultichannel(EmptyArray, OutResult, bDoFlush);
+
+		// Call flush on all the onset strength analyzers
+		for (int32 i = 0; i < OnsetStrengthAnalyzers.Num(); i++)
+		{
+			TArray<float> Strengths;
+			OnsetStrengthAnalyzers[i]->FlushAudio(Strengths);
+			OnsetStrengths[i].Append(Strengths);
+		}
+
+		// Assume that outer layers ensured that this is of correct type.
+		FOnsetNRTResult* OnsetResult = static_cast<FOnsetNRTResult*>(OutResult);
+
+		check(nullptr != OnsetResult);
+
+		if (nullptr != OnsetResult)
+		{
+			ExtractOnsetsFromOnsetStrengths(*OnsetResult);
+
+			float DurationInSeconds = static_cast<float>(NumFrames) / SampleRate;
+
+			OnsetResult->SetDurationInSeconds(DurationInSeconds);
+
+			OnsetResult->SortChronologically();
+		}
+
+		// Reset internal objects
+		Reset();
+	}
+
+	void FOnsetNRTWorker::AnalyzeWindow(const AlignedFloatBuffer& InWindow, int32 InChannelIndex, FOnsetNRTResult& OutResult)
+	{
+		TArray<float> Strengths;
+		OnsetStrengthAnalyzers[InChannelIndex]->CalculateOnsetStrengths(InWindow, Strengths);
+
+		OnsetStrengths[InChannelIndex].Append(Strengths);
+	}
+
+	void FOnsetNRTWorker::ExtractOnsetsFromOnsetStrengths(FOnsetNRTResult& OutResult)
+	{
+		for (int32 ChannelIndex = 0; ChannelIndex < OnsetStrengths.Num(); ChannelIndex++)
+		{
+			TArray<int32> OnsetIndices;
+			TArray<float>& ChannelStrengths = OnsetStrengths[ChannelIndex];
+
+			OnsetExtractIndices(Settings.PeakPickerSettings, ChannelStrengths, OnsetIndices);
+
+			// Call this in case there are no onsets. Then the results object will know about the
+			// channel, and will know there are no onsets. Otherwise it would be ambiguous whether 
+			// the channel doesn't exist or if it does exist but contains no onsets.
+			OutResult.AddChannel(ChannelIndex);
+
+			for (int32 OnsetIndex : OnsetIndices)
+			{
+				float Timestamp = FOnsetStrengthAnalyzer::GetTimestampForIndex(Settings.OnsetStrengthSettings, SampleRate, OnsetIndex);
+				OutResult.AddOnset(ChannelIndex, Timestamp, OnsetStrengths[ChannelIndex][OnsetIndex]);
+			}
+		}
+	}
+
+	void FOnsetNRTWorker::Reset()
+	{
+		// Reset internal counters
+		NumFrames = 0;
+		SlidingBuffer.Reset();
+
+		for (int32 i = 0; i < OnsetStrengthAnalyzers.Num(); i++)
+		{
+			OnsetStrengthAnalyzers[i]->Reset();
+		}
+		for (int32 i = 0; i < OnsetStrengths.Num(); i++)
+		{
+			OnsetStrengths[i].Reset();
+		}
+	}
+
+	// Name of this analyzer type.
+	FName FOnsetNRTFactory::GetName() const
+	{
+		static FName FactoryName(TEXT("OnsetNRTFactory"));
+		return FactoryName;
+	}
+
+	// Human readable name of this analyzer.
+	FString FOnsetNRTFactory::GetTitle() const
+	{
+		return TEXT("Constant Q Analyzer Non-Real-Time");
+	}
+
+	// Create a new FOnsetNRTResult.
+	TUniquePtr<IAnalyzerNRTResult> FOnsetNRTFactory::NewResult() const
+	{
+		return MakeUnique<FOnsetNRTResult>();
+	}
+
+	// Create a new FOnsetNRTWorker
+	TUniquePtr<IAnalyzerNRTWorker> FOnsetNRTFactory::NewWorker(const FAnalyzerNRTParameters& InParams, const IAnalyzerNRTSettings* InSettings) const
+	{
+		const FOnsetNRTSettings* OnsetSettings = static_cast<const FOnsetNRTSettings*>(InSettings);
+
+		check(nullptr != OnsetSettings);
+
+		return MakeUnique<FOnsetNRTWorker>(InParams, *OnsetSettings);
+	}
+}
+
diff --git a/Engine/Plugins/Runtime/AudioSynesthesia/Source/AudioSynesthesia/Private/PeakPicker.cpp b/Engine/Plugins/Runtime/AudioSynesthesia/Source/AudioSynesthesia/Private/PeakPicker.cpp
new file mode 100644
--- /dev/null
+++ b/Engine/Plugins/Runtime/AudioSynesthesia/Source/AudioSynesthesia/Private/PeakPicker.cpp
@@ -0,0 +1,66 @@
+// Copyright 1998-2019 Epic Games, Inc. All Rights Reserved.
+
+#include "PeakPicker.h"
+#include "CoreMinimal.h"
+#include "DSP/FloatArrayMath.h"
+
+namespace Audio
+{
+	FPeakPicker::FPeakPicker(const FPeakPickerSettings& InSettings)
+	:	Settings(InSettings)
+	{
+		// Clamp settings to valid values
+		Settings.NumPreMax = FMath::Max(1, Settings.NumPreMax);
+		Settings.NumPostMax = FMath::Max(1, Settings.NumPostMax);
+		Settings.NumPreMean = FMath::Max(1, Settings.NumPreMean);
+		Settings.NumPostMean = FMath::Max(1, Settings.NumPostMean);
+		Settings.NumWait = FMath::Max(1, Settings.NumWait);
+		Settings.MeanDelta = FMath::Clamp(Settings.MeanDelta, 1e-6f, INFINITY);
+	}
+
+	void FPeakPicker::PickPeaks(TArrayView<const float> InData, TArray<int32>& OutPeakIndices)
+	{
+		OutPeakIndices.Empty();
+		if (InData.Num() < 1)
+		{
+			return;
+		}
+
+		// Get max filtered version
+		TArray<float> MaxFilteredData;
+		int32 MaxFilterWindowSize = Settings.NumPreMax + Settings.NumPostMax + 1;
+		int32 MaxFilterWindowOrigin = Settings.NumPreMax;
+
+		ArrayMaxFilter(InData, MaxFilterWindowSize, MaxFilterWindowOrigin, MaxFilteredData);
+
+		
+		// Get mean filtered version
+		TArray<float> MeanFilteredData;
+		int32 MeanFilterWindowSize = Settings.NumPreMean + Settings.NumPostMean + 1;
+		int32 MeanFilterWindowOrigin = Settings.NumPreMean;
+
+		ArrayMeanFilter(InData, MeanFilterWindowSize, MeanFilterWindowOrigin, MeanFilteredData);
+
+		// Find peak locations
+		const float* InDataPtr = InData.GetData();
+		const float* MaxDataPtr = MaxFilteredData.GetData();
+		const float* MeanDataPtr = MeanFilteredData.GetData();
+
+		const int32 Num = InData.Num();
+		int32 PossiblePeakIndex = 0;
+		while (PossiblePeakIndex < Num)
+		{
+			if (InDataPtr[PossiblePeakIndex] >= MaxDataPtr[PossiblePeakIndex])
+			{
+				if (InDataPtr[PossiblePeakIndex] >= (MeanDataPtr[PossiblePeakIndex] + Settings.MeanDelta))
+				{
+					OutPeakIndices.Add(PossiblePeakIndex);
+					PossiblePeakIndex += Settings.NumWait;
+					continue;
+				}
+			}
+			PossiblePeakIndex++;
+		}
+	}
+}
+
diff --git a/Engine/Plugins/Runtime/AudioSynesthesia/Source/AudioSynesthesia/Public/AudioSynesthesiaModule.h b/Engine/Plugins/Runtime/AudioSynesthesia/Source/AudioSynesthesia/Public/AudioSynesthesiaModule.h
--- a/Engine/Plugins/Runtime/AudioSynesthesia/Source/AudioSynesthesia/Public/AudioSynesthesiaModule.h
+++ b/Engine/Plugins/Runtime/AudioSynesthesia/Source/AudioSynesthesia/Public/AudioSynesthesiaModule.h
@@ -6,8 +6,8 @@
 #include "Modules/ModuleInterface.h"
 #include "Modules/ModuleManager.h"
 #include "LoudnessNRTFactory.h"
-
-DECLARE_LOG_CATEGORY_EXTERN(LogAudioSynesthesia, Log, All);
+#include "ConstantQNRTFactory.h"
+#include "OnsetNRTFactory.h"
 
 namespace Audio
 {
@@ -21,5 +21,7 @@ namespace Audio
 
 	private:
 		FLoudnessNRTFactory LoudnessFactory;
+		FConstantQNRTFactory ConstantQFactory;
+		FOnsetNRTFactory OnsetFactory;
 	};
 }
diff --git a/Engine/Plugins/Runtime/AudioSynesthesia/Source/AudioSynesthesia/Public/ConstantQNRTFactory.h b/Engine/Plugins/Runtime/AudioSynesthesia/Source/AudioSynesthesia/Public/ConstantQNRTFactory.h
--- a/Engine/Plugins/Runtime/AudioSynesthesia/Source/AudioSynesthesia/Public/ConstantQNRTFactory.h
+++ b/Engine/Plugins/Runtime/AudioSynesthesia/Source/AudioSynesthesia/Public/ConstantQNRTFactory.h
@@ -76,22 +76,28 @@ namespace Audio
 		FConstantQNRTResult();
 
 		/** Serialize or unserialize object */
-		void Serialize(FArchive& Archive) override;
+		virtual void Serialize(FArchive& Archive) override;
 
 		/** Add a single frame of CQT data */
 		void AddFrame(int32 InChannelIndex, float InTimestamp, TArrayView<const float> InSpectrum);
 
+		/** Returns true if this result contains data for the given channel index. */
+		bool ContainsChannel(int32 InChannelIndex) const;
+
 		/** Retrieve the array of frames for a single channel of audio. */
 		const TArray<FConstantQFrame>& GetFramesForChannel(int32 InChannelIndex) const;
 
-		/** Retrieve the difference between the maximum and minimum value in the spectrum. If 
-		 *  the minimum value is less than InNoiseFloor, then the result will be the difference
-		 *  betwen the maximum value and the noise floor. 
-		 */
-		float GetChannelConstantQRange(int32 InChannelIdx, float InNoiseFloor) const;
+		/** Retrieve the difference between the maximum and minimum value in the spectrum. */
+		FFloatInterval GetChannelConstantQInterval(int32 InChannelIdx) const;
+
+		/** Retrieve an array of channel indices which exist in this result. */
+		void GetChannels(TArray<int32>& OutChannels) const;
 
-		/** Retrieve the number of channels available in the result. */
-		int32 GetNumChannels() const;
+		/** Returns the duration of the analyzed audio in seconds */
+		virtual float GetDurationInSeconds() const override;
+
+		/** Sets the duration of the analyzed audio in seconds */
+		void SetDurationInSeconds(float InDuration); 
 
 	 	/** Returns true if FConstantQFrame arrays are sorted in chronologically ascending order via their timestamp.  */
 		bool IsSortedChronologically() const;
@@ -99,8 +105,11 @@ namespace Audio
 		/** Sorts FConstantQFrame arrays in chronologically ascnding order via their timestamp.  */
 		void SortChronologically();
 
+
 	private:
 
+		float DurationInSeconds;
+
 		TMap<int32, TArray<FConstantQFrame> > ChannelCQTFrames;
 
 		TMap<int32, FFloatInterval> ChannelCQTIntervals;
@@ -129,14 +138,14 @@ namespace Audio
 		 *  InAudio is an array view of audio.
 		 *  OutResult is a pointer to a valid FConstantQNRTResult
 		 */
-		void Analyze(TArrayView<const float> InAudio, IAnalyzerNRTResult* OutResult) override;
+		virtual void Analyze(TArrayView<const float> InAudio, IAnalyzerNRTResult* OutResult) override;
 
 		/** 
 		 *  Call when analysis of audio asset is complete. 
 		 *
 		 *  OutResult must be a pointer to a valid FConstantQNRTResult. 
 		 */
-		void Finalize(IAnalyzerNRTResult* OutResult) override;
+		virtual void Finalize(IAnalyzerNRTResult* OutResult) override;
 
 	private:
 
@@ -146,6 +155,7 @@ namespace Audio
 		/** Analyze a single window of audio from a single channel */
 		void AnalyzeWindow(const AlignedFloatBuffer& InWindow, int32 InChannelIndex, FConstantQNRTResult& OutResult);
 
+		int32 NumFrames;
 		int32 NumChannels;
 		int32 NumBuffers;
 		float SampleRate;
@@ -177,18 +187,18 @@ namespace Audio
 	public:
 
 		/** Name of this analyzer type. */
-		FName GetName() const override;
+		virtual FName GetName() const override;
 
 		/** Human readable name of this analyzer. */
-		FString GetTitle() const override;
+		virtual FString GetTitle() const override;
 
 		/** Create a new FConstantQNRTResult. */
-		TUniquePtr<IAnalyzerNRTResult> NewResult() override;
+		virtual TUniquePtr<IAnalyzerNRTResult> NewResult() const override;
 
 		/** Create a new FConstantQNRTWorker 
 		 *
 		 *  InSettings must be a pointer to FConstantQNRTSetting
 		 */
-		TUniquePtr<IAnalyzerNRTWorker> NewWorker(const FAnalyzerNRTParameters& InParams, const IAnalyzerNRTSettings* InSettings) override;
+		virtual TUniquePtr<IAnalyzerNRTWorker> NewWorker(const FAnalyzerNRTParameters& InParams, const IAnalyzerNRTSettings* InSettings) const override;
 	};
 }
diff --git a/Engine/Plugins/Runtime/AudioSynesthesia/Source/AudioSynesthesia/Public/LoudnessNRTFactory.h b/Engine/Plugins/Runtime/AudioSynesthesia/Source/AudioSynesthesia/Public/LoudnessNRTFactory.h
--- a/Engine/Plugins/Runtime/AudioSynesthesia/Source/AudioSynesthesia/Public/LoudnessNRTFactory.h
+++ b/Engine/Plugins/Runtime/AudioSynesthesia/Source/AudioSynesthesia/Public/LoudnessNRTFactory.h
@@ -56,13 +56,16 @@ namespace Audio
 		/**
 		 * Defines how to serialize result.
 		 */
-		void Serialize(FArchive& Archive) override;
+		virtual void Serialize(FArchive& Archive) override;
 
 		/**
 		 * Appends an FLoudnessDatum to the container.
 		 */
 		void Add(const FLoudnessDatum& InDatum);
 
+		/** Returns true if this object data for the given channel index */
+		bool ContainsChannel(int32 InChannelIndex) const;
+
 		/**
 		 * Returns const reference to FLoudnessDatum array for individual channel.
 		 */
@@ -78,15 +81,18 @@ namespace Audio
 		 */
 		float GetLoudnessRange(float InNoiseFloor) const;
 
-		/**
-		 * Returns range in dB of loudness result given the noise floor.
-		 */
+
+		/** Returns range in dB of loudness result given the noise floor. */
 		float GetChannelLoudnessRange(int32 InChannelIdx, float InNoiseFloor) const;
 
-		/**
-		 * Get Number of Channels
-		 */
-		int32 GetNumChannels() const;
+		/** Returns the channel indices availabe in result. */
+		void GetChannels(TArray<int32>& OutChannels) const;
+
+		/** Gets the duration of the anlayzed audio. */
+		virtual float GetDurationInSeconds() const override;
+
+		/** Sets the duration of the anlayzed audio. */
+		void SetDurationInSeconds(float InDuration);
 
 		/**
 		 * Returns true if FLoudnessDatum arrays are sorted in chronologically ascending order via their timestamp.
@@ -99,6 +105,8 @@ namespace Audio
 		void SortChronologically();
 
 	private:
+		float DurationInSeconds;
+
 		TMap<int32, TArray<FLoudnessDatum> > ChannelLoudnessArrays;
 
 		TMap<int32, FFloatInterval> ChannelLoudnessIntervals;
@@ -119,12 +127,12 @@ namespace Audio
 		/**
 		 * Analyzes input sample buffer and updates result. 
 		 */
-		void Analyze(TArrayView<const float> InAudio, IAnalyzerNRTResult* OutResult) override;
+		virtual void Analyze(TArrayView<const float> InAudio, IAnalyzerNRTResult* OutResult) override;
 
 		/**
 		 * Call when all audio data has been analyzed. 
 		 */
-		void Finalize(IAnalyzerNRTResult* OutResult) override;
+		virtual void Finalize(IAnalyzerNRTResult* OutResult) override;
 
 	private:
 
@@ -134,6 +142,7 @@ namespace Audio
 		int32 NumChannels;
 		int32 NumAnalyzedBuffers;
 		int32 NumHopFrames;
+		int32 NumFrames;
 
 		float SampleRate;
 
@@ -153,19 +162,19 @@ namespace Audio
 		public:
 
 		/** Name of specific analyzer type. */
-		FName GetName() const override;
+		virtual FName GetName() const override;
 
 		/** Human readable name of analyzer. */
-		FString GetTitle() const override;
+		virtual FString GetTitle() const override;
 
 		/** Creates a new FLoudnessNRTResult */
-		TUniquePtr<IAnalyzerNRTResult> NewResult() override;
+		virtual TUniquePtr<IAnalyzerNRTResult> NewResult() const override;
 
 		/** 
 		 * Creates a new FLoudnessNRTWorker. This expects IAnalyzerNRTSettings to be a valid pointer to
 		 * a FLoudnessNRTSettings object.
 		 */
-		TUniquePtr<IAnalyzerNRTWorker> NewWorker(const FAnalyzerNRTParameters& InParams, const IAnalyzerNRTSettings* InSettings) override;
+		virtual TUniquePtr<IAnalyzerNRTWorker> NewWorker(const FAnalyzerNRTParameters& InParams, const IAnalyzerNRTSettings* InSettings) const override;
 	};
 }
 
diff --git a/Engine/Plugins/Runtime/AudioSynesthesia/Source/AudioSynesthesia/Public/OnsetAnalyzer.h b/Engine/Plugins/Runtime/AudioSynesthesia/Source/AudioSynesthesia/Public/OnsetAnalyzer.h
new file mode 100644
--- /dev/null
+++ b/Engine/Plugins/Runtime/AudioSynesthesia/Source/AudioSynesthesia/Public/OnsetAnalyzer.h
@@ -0,0 +1,106 @@
+// Copyright 1998-2019 Epic Games, Inc. All Rights Reserved.
+
+#pragma once
+
+#include "PeakPicker.h"
+#include "DSP/MelScale.h"
+#include "DSP/SlidingWindow.h"
+#include "DSP/AudioFFT.h"
+
+namespace Audio
+{
+	/** Onset strength settings */
+	struct AUDIOSYNESTHESIA_API FOnsetStrengthSettings
+	{
+		// Number of frames between strength windows
+		int32 NumHopFrames = 1024;
+
+		// Number of frames used to represent one audio window
+		int32 NumWindowFrames = 4096;
+
+		// Number of windows to lag when calculating window difference.
+		int32 ComparisonLag = 4;
+
+		// Size of FFT. Must be greater than or equal to NumWindowFrames.
+		int32 FFTSize = 4096;
+
+		// Type of analysis window to apply to audio.
+		EWindowType WindowType = EWindowType::Blackman;
+
+		// Decibel level considered to be silence.
+		float NoiseFloorDb = -60.f;
+
+		// Settings for converting FFT to mel spectrum.
+		FMelSpectrumKernelSettings MelSettings;
+	};
+
+	/** FOnsetStrengthAnalyzer
+	 *
+	 * FOnsetStrengthAnalyzer calculates the onset strength from audio. 
+	 *
+	 * Onset strength is calculated as the half wave rectificed difference between two spectral frames.
+	 * This onset strength onset analyzer uses the following approach.
+	 *
+	 * [audio]->[fft]->[mel spectrum]->[diff]->[half wave rectify]->[mean]->[onset strength]
+	 *                       |           |
+	 *                       >-[lag]-----|
+	 */
+	class AUDIOSYNESTHESIA_API FOnsetStrengthAnalyzer
+	{
+		public:
+			FOnsetStrengthAnalyzer(const FOnsetStrengthSettings& InSettings, float InSampleRate);
+
+			/** Calculates onset strengths from audio and fills OutOnsetStrengths with generated onset strengths */
+			void CalculateOnsetStrengths(TArrayView<const float> InSamples, TArray<float>& OutOnsetStrengths);
+
+			/** Call when done processing audio for an audio analyzer */
+			void FlushAudio(TArray<float>& OutEnvelopeStrengths);
+			
+			/** Call to reset internal counters and lag spectra. */
+			void Reset();
+
+			/** Converts an onset strength index into a timestamp. */
+			static float GetTimestampForIndex(const FOnsetStrengthSettings& InSettings, float InSampleRate, int32 InIndex);
+
+		private:
+
+			void AnalyzeAudio(TArrayView<const float> InSamples, TArray<float>& OutEnvelopeStrengths, bool bDoFlush = false);
+
+			// Processes single audio window
+			float GetNextOnsetStrength(TArrayView<const float> InSamples);
+
+			FOnsetStrengthSettings Settings;
+			int32 LagSpectraIndex;
+
+			TSlidingBuffer<float> SlidingBuffer;
+			AlignedFloatBuffer WorkingBuffer;
+			AlignedFloatBuffer WindowedSamples;
+			AlignedFloatBuffer FFTOutputRealData;
+			AlignedFloatBuffer FFTOutputImagData;
+			AlignedFloatBuffer FFTSpectrumBuffer;
+
+			// Array of arrays to hold lag spectra.
+			TArray< TArray<float> > PreviousSpectra;
+
+			TArray<float> MelSpectrum;
+			TArray<float> SpectrumDifference;
+
+			FWindow Window;
+			TUniquePtr<FContiguousSparse2DKernelTransform> MelTransform;
+	};
+
+	/** 
+	 * Extract onset indicies from an onset strength envelope.
+	 *
+	 * InSettings is the peak picker settings used to pick onset peaks from an onset strength envelope.
+	 * InOnsetEnvelop is an onset strength envelope generated using the FOnsetStrengthAnalyzer.
+	 * OutOnsetIndices contains the indices where onsets occur.
+	 */
+	void AUDIOSYNESTHESIA_API OnsetExtractIndices(const FPeakPickerSettings& InSettings, TArrayView<const float> InOnsetEnvelope, TArray<int32>& OutOnsetIndices);
+
+	/**
+	 * Backtracks onset indices to the beginning of the onset attack. This is useful when generating 
+	 * splice points which retain the onset attack. 
+	 */
+	void AUDIOSYNESTHESIA_API OnsetBacktrackIndices(TArrayView<const float> InOnsetEnvelope, TArrayView<const int32> InOnsetIndices, TArray<int32>& OutBacktrackedOnsetIndices);
+}
diff --git a/Engine/Plugins/Runtime/AudioSynesthesia/Source/AudioSynesthesia/Public/OnsetNRTFactory.h b/Engine/Plugins/Runtime/AudioSynesthesia/Source/AudioSynesthesia/Public/OnsetNRTFactory.h
new file mode 100644
--- /dev/null
+++ b/Engine/Plugins/Runtime/AudioSynesthesia/Source/AudioSynesthesia/Public/OnsetNRTFactory.h
@@ -0,0 +1,205 @@
+// Copyright 1998-2019 Epic Games, Inc. All Rights Reserved.
+
+#pragma once
+
+#include "IAudioAnalyzerNRTInterface.h"
+#include "OnsetAnalyzer.h"
+#include "PeakPicker.h"
+
+namespace Audio
+{
+	/** FOnsetNRTSettings
+	 *
+	 *  Settings for the Constant Q Non-Real-Time Factory.
+	 */
+	class AUDIOSYNESTHESIA_API FOnsetNRTSettings : public IAnalyzerNRTSettings
+	{
+	public:
+
+		FOnsetStrengthSettings OnsetStrengthSettings;
+
+		FPeakPickerSettings PeakPickerSettings;
+
+		/** 
+		 * If true, all channels are mixed together with equal power before analysis. Only one channel
+		 * is produced at channel index 0.
+		 * If false, then each channel is analyzed separately. 
+		 */
+		bool bDownmixToMono;
+
+		FOnsetNRTSettings()
+		:	bDownmixToMono(false)
+		{}
+	};
+
+
+	/** FOnset
+	 * 
+	 *  Contains data relating to one onset.
+	 */
+	struct AUDIOSYNESTHESIA_API FOnset
+	{
+		/** Audio channel which produced the data. */
+		int32 Channel;
+		
+		/** Timestamp in seconds referring to the center of the audio window. */
+		float Timestamp;
+
+		float Strength;
+
+		FOnset()
+		:	Channel(0)
+		,	Timestamp(0.f)
+		,	Strength(0.f)
+		{}
+
+		FOnset(int32 InChannelIndex, float InTimestamp, float InStrength)
+		:	Channel(InChannelIndex)
+		,	Timestamp(InTimestamp)
+		,	Strength(InStrength)
+		{}
+	};
+
+	/** Serialize FOnset */
+	AUDIOSYNESTHESIA_API FArchive &operator <<(FArchive& Ar, FOnset& Frame);
+
+	/** FOnsetNRTResult
+	 *
+	 *  FOnsetNRTResult is a container for the output of the FOnsetNRTWorker.
+	 */
+	class AUDIOSYNESTHESIA_API FOnsetNRTResult : public IAnalyzerNRTResult
+	{
+	public:
+
+		FOnsetNRTResult();
+
+		/** Serialize or unserialize object */
+		virtual void Serialize(FArchive& Archive) override;
+
+		/** Add a single frame of CQT data */
+		void AddOnset(int32 InChannelIndex, float InTimestamp, float InStrength);
+
+		/** Adds an empty channel. Useful if a channel exists without any onsets. */
+		void AddChannel(int32 InChannelIndex);
+
+		/** Returns true if this result contains the channel index. */
+		bool ContainsChannel(int32 InChannelIndex) const;
+
+		/** Retrieve the array of frames for a single channel of audio. */
+		const TArray<FOnset>& GetOnsetsForChannel(int32 InChannelIndex) const;
+
+		/** Retrieve the interval of the maximum and minimum strengths of all onsets in channel. */
+		FFloatInterval GetChannelOnsetInterval(int32 InChannelIdx) const;
+
+		/** Retrieve the channels available in the result. */
+		void GetChannels(TArray<int32>& OutChannels) const;
+
+		/** Returns the duration of the analyzed audio in seconds */
+		virtual float GetDurationInSeconds() const override;
+
+		/** Sets the duration of the analyzed audio in seconds */
+		void SetDurationInSeconds(float InDuration); 
+
+	 	/** Returns true if FOnset arrays are sorted in chronologically ascending order via their timestamp.  */
+		bool IsSortedChronologically() const;
+
+		/** Sorts FOnset arrays in chronologically ascnding order via their timestamp.  */
+		void SortChronologically();
+
+	private:
+
+		float DurationInSeconds;
+
+		TMap<int32, TArray<FOnset> > ChannelOnsets;
+
+		TMap<int32, FFloatInterval> ChannelOnsetStrengthIntervals;
+
+		bool bIsSortedChronologically;
+	};
+
+	/** FOnsetNRTWorker
+	 *
+	 *  FOnsetNRTWorker computes a FOnsetNRTResult from audio samples.
+	 */
+	class AUDIOSYNESTHESIA_API FOnsetNRTWorker : public IAnalyzerNRTWorker
+	{
+	public:
+
+		/** Constructor
+		 *
+		 * InParams are the parameters which describe characteristics of the input audio.
+		 * InAnalyzerSettings are the settings which control various aspects of the algorithm.
+		 */
+		FOnsetNRTWorker(const FAnalyzerNRTParameters& InParams, const FOnsetNRTSettings& InAnalyzerSettings);
+
+		/**
+		 *  Analyze audio and put results into results pointer.
+		 *
+		 *  InAudio is an array view of audio.
+		 *  OutResult is a pointer to a valid FOnsetNRTResult
+		 */
+		virtual void Analyze(TArrayView<const float> InAudio, IAnalyzerNRTResult* OutResult) override;
+
+		/** 
+		 *  Call when analysis of audio asset is complete. 
+		 *
+		 *  OutResult must be a pointer to a valid FOnsetNRTResult. 
+		 */
+		virtual void Finalize(IAnalyzerNRTResult* OutResult) override;
+
+	private:
+
+		/** Analyze audio with multiple channels interleaved. */
+		void AnalyzeMultichannel(TArrayView<const float> InAudio, IAnalyzerNRTResult* OutResult, bool bDoFlush);
+
+		void AnalyzeWindow(const AlignedFloatBuffer& InWindow, int32 InChannelIndex, FOnsetNRTResult& OutResult);
+
+		void ExtractOnsetsFromOnsetStrengths(FOnsetNRTResult& OutResult);
+
+		void Reset();
+
+		FOnsetNRTSettings Settings;
+
+		int32 NumFrames;
+		int32 NumChannels;
+		float SampleRate;
+		int32 NumWindowFrames;
+		int32 NumWindowSamples;
+		float MonoScaling;
+
+		AlignedFloatBuffer HopBuffer;
+		AlignedFloatBuffer ChannelBuffer;
+		AlignedFloatBuffer MonoBuffer;
+
+		TUniquePtr<TSlidingBuffer<float> > SlidingBuffer;
+
+		TArray< TUniquePtr<FOnsetStrengthAnalyzer> > OnsetStrengthAnalyzers;
+
+		TArray< TArray<float> > OnsetStrengths;
+
+	};
+
+	/** FOnsetNRTFactory
+	 *  
+	 *  Factory for creating FOnsetNRT workers and results
+	 */
+	class AUDIOSYNESTHESIA_API FOnsetNRTFactory : public IAnalyzerNRTFactory
+	{
+	public:
+
+		/** Name of this analyzer type. */
+		virtual FName GetName() const override;
+
+		/** Human readable name of this analyzer. */
+		virtual FString GetTitle() const override;
+
+		/** Create a new FOnsetNRTResult. */
+		virtual TUniquePtr<IAnalyzerNRTResult> NewResult() const override;
+
+		/** Create a new FOnsetNRTWorker 
+		 *
+		 *  InSettings must be a pointer to FOnsetNRTSetting
+		 */
+		virtual TUniquePtr<IAnalyzerNRTWorker> NewWorker(const FAnalyzerNRTParameters& InParams, const IAnalyzerNRTSettings* InSettings) const override;
+	};
+}
diff --git a/Engine/Plugins/Runtime/AudioSynesthesia/Source/AudioSynesthesia/Public/PeakPicker.h b/Engine/Plugins/Runtime/AudioSynesthesia/Source/AudioSynesthesia/Public/PeakPicker.h
new file mode 100644
--- /dev/null
+++ b/Engine/Plugins/Runtime/AudioSynesthesia/Source/AudioSynesthesia/Public/PeakPicker.h
@@ -0,0 +1,32 @@
+// Copyright 1998-2019 Epic Games, Inc. All Rights Reserved.
+
+#pragma once
+
+#include "CoreMinimal.h"
+
+namespace Audio
+{
+	struct FPeakPickerSettings
+	{
+		int32 NumPreMax = 1440;
+		int32 NumPostMax = 1;
+		int32 NumPreMean = 4800;
+		int32 NumPostMean = 4801;
+		int32 NumWait = 1440;
+		float MeanDelta = 0.07f;
+	};
+
+	class AUDIOSYNESTHESIA_API FPeakPicker
+	{
+		public:
+			FPeakPicker(const FPeakPickerSettings& InSettings);
+
+			void PickPeaks(TArrayView<const float> InView, TArray<int32>& OutPeakIndices);
+
+		private:
+			FPeakPickerSettings Settings;
+
+
+	};
+}
+
diff --git a/Engine/Source/Runtime/AudioAnalyzer/Classes/AudioAnalyzerNRT.h b/Engine/Source/Runtime/AudioAnalyzer/Classes/AudioAnalyzerNRT.h
--- a/Engine/Source/Runtime/AudioAnalyzer/Classes/AudioAnalyzerNRT.h
+++ b/Engine/Source/Runtime/AudioAnalyzer/Classes/AudioAnalyzerNRT.h
@@ -10,7 +10,7 @@
 #include "Sound/SoundWave.h"
 #include "AudioAnalyzerNRT.generated.h"
 
-DECLARE_DELEGATE( FAnalyzeAudioDelegate );
+DECLARE_MULTICAST_DELEGATE( FAnalyzeAudioDelegate );
 
 /** UAudioAnalyzerNRTSettings
  *
@@ -65,6 +65,14 @@ class AUDIOANALYZER_API UAudioAnalyzerNRT : public UAudioAnalyzerAsset
 	GENERATED_BODY()
 
 	public:
+		/** 
+		 * ID to keep track of results. Useful for tracking most recent result when performing
+		 * asynchronous processing.
+		 */
+		typedef int32 FResultId;
+
+		/** Thread safe shared point to result object. */
+		typedef TSharedPtr<Audio::IAnalyzerNRTResult, ESPMode::ThreadSafe> FResultSharedPtr;
 
 		/**
 		 * The USoundWave which is analyzed.
@@ -72,6 +80,10 @@ class AUDIOANALYZER_API UAudioAnalyzerNRT : public UAudioAnalyzerAsset
 		UPROPERTY(EditAnywhere, BlueprintReadOnly, Category=AudioAnalyzer)
 		USoundWave* Sound;
 
+		/** The duration of the analyzed audio in seconds. */
+		UPROPERTY(BlueprintReadOnly, Category=AudioAnalyzer)
+		float DurationInSeconds;
+
 		/**
 		 * Returns the result object generated by the associated IAudioAnalyzerNRTFactory.
 		 * The template argument must be the IAudioAnalyzerNRTResult subclass returned by
@@ -87,23 +99,40 @@ class AUDIOANALYZER_API UAudioAnalyzerNRT : public UAudioAnalyzerAsset
 			}
 			return ReturnedResult;
 		}
+
+		/**
+		 * Returns the result object generated by the associated IAudioAnalyzerNRTFactory.
+		 * The template argument must be the IAudioAnalyzerNRTResult subclass returned by
+		 * the associated IAudioAnalyzerNRTFactory.
+		 */
+		template<class ResultType>
+		TSharedPtr<const ResultType, ESPMode::ThreadSafe> GetResult() const
+		{
+			TSharedPtr<const ResultType, ESPMode::ThreadSafe> ReturnedResult;
+			{
+				FScopeLock ResultLock(&ResultCriticalSection);
+				ReturnedResult = StaticCastSharedPtr<const ResultType>(Result); 
+			}
+			return ReturnedResult;
+		}
 	 
 		/** 
 		 * Implementations can override this method to create settings objects
 		 * specific for their analyzer. 
 		 */
-		virtual TUniquePtr<Audio::IAnalyzerNRTSettings> GetSettings();
+		virtual TUniquePtr<Audio::IAnalyzerNRTSettings> GetSettings(const float InSampleRate, const int32 InNumChannels) const;
 
 		/**
 		 * Performs serialization of results.
 		 */
 		virtual void Serialize(FArchive& Ar) override;
 
-		/* TODO: progress/status for calculation */
-
 #if WITH_EDITOR
 
-		void SetResult(TSharedPtr<Audio::IAnalyzerNRTResult, ESPMode::ThreadSafe> NewResult);
+		void SetResult(FResultSharedPtr NewResult);
+
+		/** This will only store the result if the passed InResultId matches the CurrentResultId. */
+		void SetResultIfLatest(FResultSharedPtr NewResult, FResultId InResultId);
 
 		/**
 		 * Called before a UPROPERTY of this class is edited. This checks to see
@@ -133,6 +162,16 @@ class AUDIOANALYZER_API UAudioAnalyzerNRT : public UAudioAnalyzerAsset
 
 
 	protected:
+
+#if WITH_EDITOR
+
+		/** This sets the AnalyzeAudio callback to the delegate in settings. */
+		void SetSettingsDelegate(UAudioAnalyzerNRTSettings* InSettings);
+
+		/** This removes the existing delegate handle from the current settings object */
+		void RemoveSettingsDelegate(UAudioAnalyzerNRTSettings* InSettings);
+#endif 
+
 		/* Subclasses must override this method in order to inform this object which AnalyzerNRTFactory to use for analysis */
 		virtual FName GetAnalyzerNRTFactoryName() const PURE_VIRTUAL(UAudioAnalyzerNRT::GetAnalyzerNRTFactoryName, return FName(););
 
@@ -142,6 +181,15 @@ class AUDIOANALYZER_API UAudioAnalyzerNRT : public UAudioAnalyzerAsset
 		UAudioAnalyzerNRTSettings* GetSettingsFromProperty(UProperty* Property);
 
 		TSharedPtr<Audio::IAnalyzerNRTResult, ESPMode::ThreadSafe> Result;
-		FCriticalSection ResultCriticalSection;
+
+		// Critical section is mutable so it can still be locked within const methods.
+		mutable FCriticalSection ResultCriticalSection;
+
+		// Result id of the current result. 
+		TAtomic<FResultId> CurrentResultId;
+
+#if WITH_EDITOR
+		TMap<UAudioAnalyzerNRTSettings*, FDelegateHandle> AnalyzeAudioDelegateHandles;
+#endif 
 };
 
diff --git a/Engine/Source/Runtime/AudioAnalyzer/Private/AudioAnalyzerNRT.cpp b/Engine/Source/Runtime/AudioAnalyzer/Private/AudioAnalyzerNRT.cpp
--- a/Engine/Source/Runtime/AudioAnalyzer/Private/AudioAnalyzerNRT.cpp
+++ b/Engine/Source/Runtime/AudioAnalyzer/Private/AudioAnalyzerNRT.cpp
@@ -6,6 +6,7 @@
 #include "SampleBuffer.h"
 #include "Async/Async.h"
 
+
 #if WITH_EDITOR
 
 namespace 
@@ -17,11 +18,13 @@ namespace
 		public:
 			FAudioAnalyzeTask(
 					TWeakObjectPtr<UAudioAnalyzerNRT> InAnalyzerUObject, 
+					const UAudioAnalyzerNRT::FResultId InResultId,
 					TUniquePtr<Audio::FAnalyzerNRTBatch>&& InAnalyzerFacade, 
 					TArray<uint8>&& InRawWaveData,
 					int32 InNumChannels,
 					float InSampleRate)
 			: AnalyzerUObject(InAnalyzerUObject)
+			, ResultId(InResultId)
 			, AnalyzerFacade(MoveTemp(InAnalyzerFacade))
 			, RawWaveData(MoveTemp(InRawWaveData))
 			, NumChannels(InNumChannels)
@@ -32,16 +35,12 @@ namespace
 			{
 				TUniquePtr<Audio::IAnalyzerNRTResult> Result = AnalyzerFacade->AnalyzePCM16Audio(RawWaveData, NumChannels, SampleRate);
 
-				// Make sharedptr so can be passed by value to async lambda
-				TSharedPtr<Audio::IAnalyzerNRTResult, ESPMode::ThreadSafe> ResultPtr(Result.Release());
-				// Make local so can be passed to lambda without having to capture 'this'
-				TWeakObjectPtr<UAudioAnalyzerNRT> Analyzer = AnalyzerUObject;
-
+				UAudioAnalyzerNRT::FResultSharedPtr ResultPtr(Result.Release());
 				// Set value on game thread.
-				AsyncTask(ENamedThreads::GameThread, [Analyzer, ResultPtr]() {
+				AsyncTask(ENamedThreads::GameThread, [Analyzer = AnalyzerUObject, ThisResultId = ResultId, ResultPtr]() {
 					if (Analyzer.IsValid())
 					{
-						Analyzer->SetResult(ResultPtr);
+						Analyzer->SetResultIfLatest(ResultPtr, ThisResultId);
 					}
 				});
 			}
@@ -50,6 +49,7 @@ namespace
 
 		private:
 			TWeakObjectPtr<UAudioAnalyzerNRT> AnalyzerUObject;
+			const UAudioAnalyzerNRT::FResultId ResultId;
 			TUniquePtr<Audio::FAnalyzerNRTBatch> AnalyzerFacade;
 			TArray<uint8> RawWaveData;
 			int32 NumChannels;
@@ -67,14 +67,14 @@ void UAudioAnalyzerNRTSettings::PostEditChangeProperty (struct FPropertyChangedE
 
 	if (ShouldEventTriggerAnalysis(PropertyChangedEvent))
 	{
-		AnalyzeAudioDelegate.ExecuteIfBound();
+		AnalyzeAudioDelegate.Broadcast();
 	}
 }
 
 bool UAudioAnalyzerNRTSettings::ShouldEventTriggerAnalysis(struct FPropertyChangedEvent & PropertyChangeEvent)
 {
-	// By default, all changes to settings will trigger analysis.
-	return true;
+	// By default, all non-interactive changes to settings will trigger analysis.
+	return PropertyChangeEvent.ChangeType != EPropertyChangeType::Interactive;
 }
 
 
@@ -92,7 +92,7 @@ void UAudioAnalyzerNRT::PreEditChange(UProperty* PropertyAboutToChange)
 
 	if (Settings)
 	{
-		Settings->AnalyzeAudioDelegate.Unbind();
+		RemoveSettingsDelegate(Settings);
 	}
 }
 
@@ -106,7 +106,7 @@ void UAudioAnalyzerNRT::PostEditChangeProperty (struct FPropertyChangedEvent & P
 	if (Settings)
 	{
 		// If it was a UAudioAnalyzerNRTSettings object, bind the FAnalyzeAudioDelegate
-		Settings->AnalyzeAudioDelegate.BindUObject(this, &UAudioAnalyzerNRT::AnalyzeAudio);
+		SetSettingsDelegate(Settings);
 	}
 
 	if (ShouldEventTriggerAnalysis(PropertyChangedEvent))
@@ -125,6 +125,9 @@ void UAudioAnalyzerNRT::AnalyzeAudio()
 {
 	TSharedPtr<Audio::IAnalyzerNRTResult, ESPMode::ThreadSafe> NewResult;
 
+	// Create a new result id for this result.
+	FResultId ThisResultId = ++CurrentResultId;
+
 	if (nullptr != Sound)
 	{
 		// Read audio while Sound object is assured safe. 
@@ -152,13 +155,13 @@ void UAudioAnalyzerNRT::AnalyzeAudio()
 		}
 		
 		// Create analyzer helper object
-		TUniquePtr<Audio::FAnalyzerNRTBatch> BatchAnalyzer = MakeUnique<Audio::FAnalyzerNRTBatch>(GetSettings(), GetAnalyzerNRTFactoryName());
+		TUniquePtr<Audio::FAnalyzerNRTBatch> BatchAnalyzer = MakeUnique<Audio::FAnalyzerNRTBatch>(GetSettings(SampleRate, NumChannels), GetAnalyzerNRTFactoryName());
 
 		// Use weak reference in case this object is deleted before analysis is done
 		TWeakObjectPtr<UAudioAnalyzerNRT> AnalyzerPtr(this);
 		
 		// Create and start async task. Parentheses avoids memory leak warnings from static analysis.
-		(new FAutoDeleteAsyncTask<FAudioAnalyzeTask>(AnalyzerPtr, MoveTemp(BatchAnalyzer), MoveTemp(RawWaveData), NumChannels, SampleRate))->StartBackgroundTask();
+		(new FAutoDeleteAsyncTask<FAudioAnalyzeTask>(AnalyzerPtr, ThisResultId, MoveTemp(BatchAnalyzer), MoveTemp(RawWaveData), NumChannels, SampleRate))->StartBackgroundTask();
 	}
 	else
 	{
@@ -194,10 +197,39 @@ UAudioAnalyzerNRTSettings* UAudioAnalyzerNRT::GetSettingsFromProperty(UProperty*
 	return nullptr;
 }
 
-void UAudioAnalyzerNRT::SetResult(TSharedPtr<Audio::IAnalyzerNRTResult, ESPMode::ThreadSafe> NewResult)
+void UAudioAnalyzerNRT::SetResult(FResultSharedPtr NewResult)
 {
 	FScopeLock ResultLock(&ResultCriticalSection);
 	Result = NewResult;
+	if (Result.IsValid())
+	{
+		DurationInSeconds = Result->GetDurationInSeconds();
+	}
+	else
+	{
+		DurationInSeconds = 0.f;
+	}
+	Modify();
+}
+
+void UAudioAnalyzerNRT::SetResultIfLatest(FResultSharedPtr NewResult, FResultId InResultId)
+{
+	FScopeLock ResultLock(&ResultCriticalSection);
+	const FResultId ResultId = CurrentResultId.Load();
+
+	if (ResultId == InResultId)
+	{
+		Result = NewResult;
+		if (Result.IsValid())
+		{
+			DurationInSeconds = Result->GetDurationInSeconds();
+		}
+		else
+		{
+			DurationInSeconds = 0.f;
+		}
+		Modify();
+	}
 }
 
 #endif
@@ -234,8 +266,47 @@ void UAudioAnalyzerNRT::Serialize(FArchive& Ar)
 	}
 }
 
-TUniquePtr<Audio::IAnalyzerNRTSettings> UAudioAnalyzerNRT::GetSettings()
+TUniquePtr<Audio::IAnalyzerNRTSettings> UAudioAnalyzerNRT::GetSettings(const float InSampleRate, const int32 InNumChannels) const
 {
 	return MakeUnique<Audio::IAnalyzerNRTSettings>();
 }
 
+#if WITH_EDITOR
+void UAudioAnalyzerNRT::SetSettingsDelegate(UAudioAnalyzerNRTSettings* InSettings)
+{
+	if (InSettings)
+	{
+		if (AnalyzeAudioDelegateHandles.Contains(InSettings))
+		{
+			// Avoid setting delegate more tha once
+			return;
+		}
+
+		FDelegateHandle DelegateHandle = InSettings->AnalyzeAudioDelegate.AddUObject(this, &UAudioAnalyzerNRT::AnalyzeAudio);
+
+		if (DelegateHandle.IsValid())
+		{
+			AnalyzeAudioDelegateHandles.Add(InSettings, DelegateHandle);
+		}
+	}
+}
+
+void UAudioAnalyzerNRT::RemoveSettingsDelegate(UAudioAnalyzerNRTSettings* InSettings)
+{
+	if (InSettings)
+	{
+		if (AnalyzeAudioDelegateHandles.Contains(InSettings))
+		{
+			FDelegateHandle DelegateHandle = AnalyzeAudioDelegateHandles[InSettings];
+
+			if (DelegateHandle.IsValid())
+			{
+				InSettings->AnalyzeAudioDelegate.Remove(DelegateHandle);
+			}
+
+			AnalyzeAudioDelegateHandles.Remove(InSettings);
+		}
+	}
+}
+
+#endif
\ No newline at end of file
diff --git a/Engine/Source/Runtime/AudioAnalyzer/Private/IAudioAnalyzerNRTInterface.cpp b/Engine/Source/Runtime/AudioAnalyzer/Private/IAudioAnalyzerNRTInterface.cpp
--- a/Engine/Source/Runtime/AudioAnalyzer/Private/IAudioAnalyzerNRTInterface.cpp
+++ b/Engine/Source/Runtime/AudioAnalyzer/Private/IAudioAnalyzerNRTInterface.cpp
@@ -15,7 +15,7 @@ namespace Audio
 		Serialize(DeserializedResult);
 	}
 	
-	FString IAnalyzerNRTResult::ToString()
+	FString IAnalyzerNRTResult::ToString() const
 	{
 		return FString(TEXT("To use, override ToString() in this implementation of IAnalyzerNRTResult."));
 	}
diff --git a/Engine/Source/Runtime/AudioAnalyzer/Public/IAudioAnalyzerNRTInterface.h b/Engine/Source/Runtime/AudioAnalyzer/Public/IAudioAnalyzerNRTInterface.h
--- a/Engine/Source/Runtime/AudioAnalyzer/Public/IAudioAnalyzerNRTInterface.h
+++ b/Engine/Source/Runtime/AudioAnalyzer/Public/IAudioAnalyzerNRTInterface.h
@@ -56,7 +56,10 @@ namespace Audio
 		virtual void CopyFrom(IAnalyzerNRTResult* SourceResult);
 
 		// This can be overridden to return a string description of this results struct.
-		virtual FString ToString();
+		virtual FString ToString() const;
+
+		// This must be overridden to return the duration of the original audio analyzed.
+		virtual float GetDurationInSeconds() const = 0;
 	};
 
 	/** IAnalyzerNRTWorker
@@ -101,11 +104,11 @@ namespace Audio
 		virtual FString GetTitle() const;
 
 		// Create a new result.
-		virtual TUniquePtr<IAnalyzerNRTResult> NewResult() = 0;
+		virtual TUniquePtr<IAnalyzerNRTResult> NewResult() const = 0;
 
 		// Convenience function to create a new shared result by calling NewResult.
 		template<ESPMode Mode = ESPMode::Fast>
-		TSharedPtr<IAnalyzerNRTResult, Mode> NewResultShared()
+		TSharedPtr<IAnalyzerNRTResult, Mode> NewResultShared() const
 		{
 			TUniquePtr<Audio::IAnalyzerNRTResult> Result = NewResult();
 
@@ -113,7 +116,7 @@ namespace Audio
 		}
 
 		// Create a new worker.
-		virtual TUniquePtr<IAnalyzerNRTWorker> NewWorker(const FAnalyzerNRTParameters& InParams, const IAnalyzerNRTSettings* InSettings) = 0;
+		virtual TUniquePtr<IAnalyzerNRTWorker> NewWorker(const FAnalyzerNRTParameters& InParams, const IAnalyzerNRTSettings* InSettings) const = 0;
 	};
 }
 
diff --git a/Engine/Source/Runtime/SignalProcessing/Private/FloatArrayMath.cpp b/Engine/Source/Runtime/SignalProcessing/Private/FloatArrayMath.cpp
--- a/Engine/Source/Runtime/SignalProcessing/Private/FloatArrayMath.cpp
+++ b/Engine/Source/Runtime/SignalProcessing/Private/FloatArrayMath.cpp
@@ -35,6 +35,27 @@ namespace Audio
 		}
 	}
 
+	void ArrayMean(TArrayView<const float> InView, float& OutMean)
+	{
+		OutMean = 0.f;
+
+		const int32 Num = InView.Num();
+
+		if (Num < 1)
+		{
+			return;
+		}
+
+		const float* DataPtr = InView.GetData();
+
+		for (int32 i = 0; i < Num; i++)
+		{
+			OutMean += DataPtr[i];
+		}
+
+		OutMean /= static_cast<float>(Num);
+	}
+
 	void ArrayMeanFilter(TArrayView<const float> InView, int32 WindowSize, int32 WindowOrigin, TArray<float>& OutData)
 	{
 		// a quick but sinful implementation of a mean filter. encourages floating point rounding errors. 
@@ -196,6 +217,39 @@ namespace Audio
 		OutEuclideanNorm = FMath::Sqrt(OutEuclideanNorm);
 	}
 
+	void ArrayAbsInPlace(TArrayView<float> InView)
+	{
+		const int32 Num = InView.Num();
+		float* Data = InView.GetData();
+
+		for (int32 i = 0; i < Num; i++)
+		{
+			Data[i] = FMath::Abs(Data[i]);
+		}
+	}
+
+	void ArrayClampMinInPlace(TArrayView<float> InView, float InMin)
+	{
+		const int32 Num = InView.Num();
+		float* Data = InView.GetData();
+
+		for (int32 i = 0; i < Num; i++)
+		{
+			Data[i] = FMath::Max(InMin, Data[i]);
+		}
+	}
+
+	void ArrayClampMaxInPlace(TArrayView<float> InView, float InMax)
+	{
+		const int32 Num = InView.Num();
+		float* Data = InView.GetData();
+
+		for (int32 i = 0; i < Num; i++)
+		{
+			Data[i] = FMath::Min(InMax, Data[i]);
+		}
+	}
+
 	void ArrayClampInPlace(TArrayView<float> InView, float InMin, float InMax)
 	{
 		const int32 Num = InView.Num();
@@ -207,6 +261,44 @@ namespace Audio
 		}
 	}
 
+	void ArrayMinMaxNormalize(TArrayView<const float> InView, TArray<float>& OutArray)
+	{
+		const int32 Num = InView.Num();
+		OutArray.Reset(Num);
+
+		if (Num < 1)
+		{
+			return;
+		}
+
+		OutArray.AddUninitialized(Num);
+
+		const float* InDataPtr = InView.GetData();
+		float MaxValue = InDataPtr[0];
+		float MinValue = InDataPtr[0];
+
+		// determine min and max
+		for (int32 i = 1; i < Num; i++)
+		{
+			if (InDataPtr[i] < MinValue)
+			{
+				MinValue = InDataPtr[i];
+			}
+			else if (InDataPtr[i] > MaxValue)
+			{
+				MaxValue = InDataPtr[i];
+			}
+		}
+
+		// Normalize data by subtracting minimum value and dividing by range
+		float* OutDataPtr = OutArray.GetData();
+		float Scale = 1.f / FMath::Max(SMALL_NUMBER, MaxValue - MinValue);
+		for (int32 i = 0; i < Num; i++)
+		{
+			OutDataPtr[i] = (InDataPtr[i] - MinValue) * Scale;
+		}
+	}
+
 	void ArrayMultiplyByConstantInPlace(TArrayView<float> InView, float InMultiplier)
 	{
 		const int32 Num = InView.Num();
@@ -227,6 +319,31 @@ namespace Audio
 		}
 	}
 
+	void ArraySubtract(TArrayView<const float> InMinuend, TArrayView<const float> InSubtrahend, TArray<float>& OutArray)
+	{
+		const int32 Num = InMinuend.Num();
+
+		checkf(Num == InSubtrahend.Num(), TEXT("InMinuend and InSubtrahend must have equal Num elements (%d vs %d)"), Num, InSubtrahend.Num());
+
+		OutArray.Reset(Num);
+
+		if (Num < 1)
+		{
+			return;
+		}
+
+		OutArray.AddUninitialized(Num);
+		
+		const float* MinuendPtr = InMinuend.GetData();
+		const float* SubtrahendPtr = InSubtrahend.GetData();
+		float* OutPtr = OutArray.GetData();
+
+		for (int32 i = 0; i < Num; i++)
+		{
+			OutPtr[i] = MinuendPtr[i] - SubtrahendPtr[i];
+		}
+	}
+
 	void ArrayMagnitudeToDecibelInPlace(TArrayView<float> InView)
 	{
 		const int32 Num = InView.Num();
diff --git a/Engine/Source/Runtime/SignalProcessing/Public/DSP/FloatArrayMath.h b/Engine/Source/Runtime/SignalProcessing/Public/DSP/FloatArrayMath.h
--- a/Engine/Source/Runtime/SignalProcessing/Public/DSP/FloatArrayMath.h
+++ b/Engine/Source/Runtime/SignalProcessing/Public/DSP/FloatArrayMath.h
@@ -14,6 +14,13 @@ namespace Audio
 	 */
 	SIGNALPROCESSING_API void ArrayCumulativeSum(TArrayView<const float> InView, TArray<float>& OutData);
 
+	/** Mean of array. Equivalent to Sum(InView) / InView.Num()
+	 *
+	 *  InView contains data to be analyzed.
+	 *  OutMean contains the result.
+	 */
+	SIGNALPROCESSING_API void ArrayMean(TArrayView<const float> InView, float& OutMean);
+
 	/** Mean filter of array.
 	 *
 	 *  Note: Uses standard biased mean estimator of Sum(x) / Count(x).
@@ -40,6 +47,26 @@ namespace Audio
 	/** Computes the EuclideanNorm of the InView. Same as calculating the energy in window. */
 	SIGNALPROCESSING_API void ArrayGetEuclideanNorm(TArrayView<const float> InView, float& OutEuclideanNorm);
 
+	/** Absolute value of array elements in place.
+	 *
+	 *  InView contains the data to be manipulated.
+	 */
+	SIGNALPROCESSING_API void ArrayAbsInPlace(TArrayView<float> InView);
+
+	/** Clamp minimum value of array in place.
+	 *
+	 *  InView contains data to be clamped.
+	 *  InMin contains the minimum value allowable in InView.
+	 */
+	SIGNALPROCESSING_API void ArrayClampMinInPlace(TArrayView<float> InView, float InMin);
+
+	/** Clamp maximum value of array in place.
+	 *
+	 *  InView contains data to be clamped.
+	 *  InMax contains the maximum value allowable in InView.
+	 */
+	SIGNALPROCESSING_API void ArrayClampMaxInPlace(TArrayView<float> InView, float InMax);
+
 	/** Clamp values in an array.
 	 *
 	 *  InView is a view of a float array to be clamped.
@@ -48,9 +75,24 @@ namespace Audio
 	 */
 	SIGNALPROCESSING_API void ArrayClampInPlace(TArrayView<float> InView, float InMin, float InMax);
 
+	/** Scale an array so the minimum is 0 and the maximum is 1
+	 *
+	 *  InView is the view of a float array with the input data.
+	 *  OutArray is an array which will hold the normalized data.
+	 */ 
+	SIGNALPROCESSING_API void ArrayMinMaxNormalize(TArrayView<const float> InView, TArray<float>& OutArray);
+
 	/** Multiplies each element in InView by InMultiplier */
 	SIGNALPROCESSING_API void ArrayMultiplyByConstantInPlace(TArrayView<float> InView, float InMultiplier);
 
+	/** Subract arrays element-wise. OutArray = InMinuend - InSubtrahend
+	 *
+	 *  InMinuend is the array of data to be subtracted from.
+	 *  InSubtrahend is the array of data to subtract.
+	 *  OutArray is the array which holds the result.
+	 */
+	SIGNALPROCESSING_API void ArraySubtract(TArrayView<const float> InMinuend, TArrayView<const float> InSubtrahend, TArray<float>& OutArray);
+
 	/** Subtract value from each element in InView */
 	SIGNALPROCESSING_API void ArraySubtractByConstantInPlace(TArrayView<float> InView, float InSubtrahend);
 
