Changing AudioSynesthesia load phase from default to preinit

diff --git a/Engine/Plugins/Runtime/AudioSynesthesia/AudioSynesthesia.uplugin b/Engine/Plugins/Runtime/AudioSynesthesia/AudioSynesthesia.uplugin
--- a/Engine/Plugins/Runtime/AudioSynesthesia/AudioSynesthesia.uplugin
+++ b/Engine/Plugins/Runtime/AudioSynesthesia/AudioSynesthesia.uplugin
@@ -19,7 +19,7 @@
 		{
 			"Name" : "AudioSynesthesia",
 			"Type" : "Runtime",
-			"LoadingPhase" : "Default"
+			"LoadingPhase" : "PreDefault"
 		},
 		{
 			"Name" : "AudioSynesthesiaEditor",
