Fix Synesthesia types not displaying names correctly in Content Browser filters
#rnx

diff --git a/Engine/Plugins/Runtime/AudioSynesthesia/Source/AudioSynesthesia/Classes/AudioSynesthesiaNRT.h b/Engine/Plugins/Runtime/AudioSynesthesia/Source/AudioSynesthesia/Classes/AudioSynesthesiaNRT.h
--- a/Engine/Plugins/Runtime/AudioSynesthesia/Source/AudioSynesthesia/Classes/AudioSynesthesiaNRT.h
+++ b/Engine/Plugins/Runtime/AudioSynesthesia/Source/AudioSynesthesia/Classes/AudioSynesthesiaNRT.h
@@ -15,12 +15,8 @@ class AUDIOSYNESTHESIA_API UAudioSynesthesiaNRTSettings : public UAudioAnalyzerN
 
 	public:
 
-		FText GetAssetActionName() const override;
-
 		const TArray<FText>& GetAssetActionSubmenus() const;
 
-		UClass* GetSupportedClass() const override;
-
 		FColor GetTypeColor() const override;
 };
 
@@ -35,12 +31,8 @@ class AUDIOSYNESTHESIA_API UAudioSynesthesiaNRT : public UAudioAnalyzerNRT
 
 	public:
 
-		FText GetAssetActionName() const override;
-
 		const TArray<FText>& GetAssetActionSubmenus() const;
 
-		UClass* GetSupportedClass() const override;
-
 		FColor GetTypeColor() const override;
 };
 
diff --git a/Engine/Plugins/Runtime/AudioSynesthesia/Source/AudioSynesthesia/Classes/ConstantQNRT.h b/Engine/Plugins/Runtime/AudioSynesthesia/Source/AudioSynesthesia/Classes/ConstantQNRT.h
--- a/Engine/Plugins/Runtime/AudioSynesthesia/Source/AudioSynesthesia/Classes/ConstantQNRT.h
+++ b/Engine/Plugins/Runtime/AudioSynesthesia/Source/AudioSynesthesia/Classes/ConstantQNRT.h
@@ -116,6 +116,10 @@ class AUDIOSYNESTHESIA_API UConstantQNRTSettings : public UAudioSynesthesiaNRTSe
 
 		/** Convert UConstantQNRTSettings to FConstantQNRTSettings */
 		TUniquePtr<Audio::IAnalyzerNRTSettings> GetSettings(const float InSampleRate, const int32 InNumChannels) const;
+
+		virtual FText GetAssetActionName() const override;
+
+		virtual UClass* GetSupportedClass() const override;
 };
 
 
@@ -148,6 +152,10 @@ class AUDIOSYNESTHESIA_API UConstantQNRT : public UAudioSynesthesiaNRT
 		/** Convert ULoudnessNRTSettings to FLoudnessNRTSettings */
  		virtual TUniquePtr<Audio::IAnalyzerNRTSettings> GetSettings(const float InSampleRate, const int32 InNumChannels) const override;
 
+		virtual FText GetAssetActionName() const override;
+
+		virtual UClass* GetSupportedClass() const override;
+
 	protected:
 
 		/** Return the name of the IAudioAnalyzerNRTFactory associated with this UAudioAnalyzerNRT */
diff --git a/Engine/Plugins/Runtime/AudioSynesthesia/Source/AudioSynesthesia/Classes/LoudnessNRT.h b/Engine/Plugins/Runtime/AudioSynesthesia/Source/AudioSynesthesia/Classes/LoudnessNRT.h
--- a/Engine/Plugins/Runtime/AudioSynesthesia/Source/AudioSynesthesia/Classes/LoudnessNRT.h
+++ b/Engine/Plugins/Runtime/AudioSynesthesia/Source/AudioSynesthesia/Classes/LoudnessNRT.h
@@ -58,6 +58,10 @@ class AUDIOSYNESTHESIA_API ULoudnessNRTSettings : public UAudioSynesthesiaNRTSet
 
 		/** Convert ULoudnessNRTSettings to FLoudnessNRTSettings */
 		TUniquePtr<Audio::IAnalyzerNRTSettings> GetSettings(const float InSampleRate, const int32 InNumChannels) const;
+
+		virtual FText GetAssetActionName() const override;
+
+		virtual UClass* GetSupportedClass() const override;
 };
 
 
@@ -108,6 +112,10 @@ class AUDIOSYNESTHESIA_API ULoudnessNRT : public UAudioSynesthesiaNRT
 		/** Convert ULoudnessNRTSettings to FLoudnessNRTSettings */
  		TUniquePtr<Audio::IAnalyzerNRTSettings> GetSettings(const float InSampleRate, const int32 InNumChannels) const override;
 
+		virtual FText GetAssetActionName() const override;
+
+		virtual UClass* GetSupportedClass() const override;
+
 	protected:
 
 		/** Return the name of the IAudioAnalyzerNRTFactory associated with this UAudioAnalyzerNRT */
diff --git a/Engine/Plugins/Runtime/AudioSynesthesia/Source/AudioSynesthesia/Classes/OnsetNRT.h b/Engine/Plugins/Runtime/AudioSynesthesia/Source/AudioSynesthesia/Classes/OnsetNRT.h
--- a/Engine/Plugins/Runtime/AudioSynesthesia/Source/AudioSynesthesia/Classes/OnsetNRT.h
+++ b/Engine/Plugins/Runtime/AudioSynesthesia/Source/AudioSynesthesia/Classes/OnsetNRT.h
@@ -36,10 +36,13 @@ class AUDIOSYNESTHESIA_API UOnsetNRTSettings : public UAudioSynesthesiaNRTSettin
 		/** Starting frequency for onset anlaysis. */
 		UPROPERTY(EditAnywhere, BlueprintReadOnly, Category=AudioAnalyzer, meta = (ClampMin = "20.0", ClampMax = "20000"))
 		float MaximumFrequency;
-
-
+		
 		/** Convert UOnsetNRTSettings to FOnsetNRTSettings */
 		TUniquePtr<Audio::IAnalyzerNRTSettings> GetSettings(const float InSampleRate, const int32 InNumChannels) const;
+
+		virtual FText GetAssetActionName() const override;
+
+		virtual UClass* GetSupportedClass() const override;
 };
 
 
@@ -72,6 +75,10 @@ class AUDIOSYNESTHESIA_API UOnsetNRT : public UAudioSynesthesiaNRT
 		/** Convert ULoudnessNRTSettings to FLoudnessNRTSettings */
  		virtual TUniquePtr<Audio::IAnalyzerNRTSettings> GetSettings(const float InSampleRate, const int32 InNumChannels) const override;
 
+		virtual FText GetAssetActionName() const override;
+
+		virtual UClass* GetSupportedClass() const override;
+
 	protected:
 
 		/** Return the name of the IAudioAnalyzerNRTFactory associated with this UAudioAnalyzerNRT */
diff --git a/Engine/Plugins/Runtime/AudioSynesthesia/Source/AudioSynesthesia/Private/AudioSynesthesiaNRT.cpp b/Engine/Plugins/Runtime/AudioSynesthesia/Source/AudioSynesthesia/Private/AudioSynesthesiaNRT.cpp
--- a/Engine/Plugins/Runtime/AudioSynesthesia/Source/AudioSynesthesia/Private/AudioSynesthesiaNRT.cpp
+++ b/Engine/Plugins/Runtime/AudioSynesthesia/Source/AudioSynesthesia/Private/AudioSynesthesiaNRT.cpp
@@ -3,11 +3,6 @@
 #include "AudioSynesthesiaNRT.h"
 
 
-FText UAudioSynesthesiaNRTSettings::GetAssetActionName() const
-{
-	return NSLOCTEXT("AssetTypeActions", "AssetTypeActions_AssetSoundSynesthesiaNRTSettings", "Synesthesia NRT Settings");
-}
-
 const TArray<FText>& UAudioSynesthesiaNRTSettings::GetAssetActionSubmenus() const
 {
 	static const TArray<FText> SubMenus
@@ -18,21 +13,11 @@ const TArray<FText>& UAudioSynesthesiaNRTSettings::GetAssetActionSubmenus() cons
 	return SubMenus;
 }
 
-UClass* UAudioSynesthesiaNRTSettings::GetSupportedClass() const 
-{
-	return UAudioSynesthesiaNRTSettings::StaticClass();
-}
-
 FColor UAudioSynesthesiaNRTSettings::GetTypeColor() const
 {
 	return FColor(200.0f, 150.0f, 200.0f);
 }
 
-FText UAudioSynesthesiaNRT::GetAssetActionName() const
-{
-	return NSLOCTEXT("AssetTypeActions", "AssetTypeActions_AssetSoundSynesthesiaNRT", "Synesthesia NRT");
-}
-
 const TArray<FText>& UAudioSynesthesiaNRT::GetAssetActionSubmenus() const
 {
 	static const TArray<FText> SubMenus
@@ -43,11 +28,6 @@ const TArray<FText>& UAudioSynesthesiaNRT::GetAssetActionSubmenus() const
 	return SubMenus;
 }
 
-UClass* UAudioSynesthesiaNRT::GetSupportedClass() const
-{
-	return UAudioSynesthesiaNRT::StaticClass();
-}
-
 FColor UAudioSynesthesiaNRT::GetTypeColor() const
 {
 	return FColor(200.0f, 150.0f, 200.0f);
diff --git a/Engine/Plugins/Runtime/AudioSynesthesia/Source/AudioSynesthesia/Private/ConstantQNRT.cpp b/Engine/Plugins/Runtime/AudioSynesthesia/Source/AudioSynesthesia/Private/ConstantQNRT.cpp
--- a/Engine/Plugins/Runtime/AudioSynesthesia/Source/AudioSynesthesia/Private/ConstantQNRT.cpp
+++ b/Engine/Plugins/Runtime/AudioSynesthesia/Source/AudioSynesthesia/Private/ConstantQNRT.cpp
@@ -154,6 +154,15 @@ TUniquePtr<Audio::IAnalyzerNRTSettings> UConstantQNRTSettings::GetSettings(const
 	return Settings;
 }
 
+FText UConstantQNRTSettings::GetAssetActionName() const
+{
+	return NSLOCTEXT("AssetTypeActions", "AssetTypeActions_AssetSoundSynesthesiaConstantQNRTSettings", "Synesthesia NRT Settings (ConstantQ)");
+}
+
+UClass* UConstantQNRTSettings::GetSupportedClass() const
+{
+	return UConstantQNRTSettings::StaticClass();
+}
 
 /***************************************************************************/
 /*********************        UConstantQNRT         ************************/
@@ -243,6 +252,16 @@ TUniquePtr<Audio::IAnalyzerNRTSettings> UConstantQNRT::GetSettings(const float I
 	return AnalyzerSettings;
 }
 
+FText UConstantQNRT::GetAssetActionName() const
+{
+	return NSLOCTEXT("AssetTypeActions", "AssetTypeActions_AssetSoundSynesthesiaConstantQNRT", "Synesthesia NRT (ConstantQ)");
+}
+
+UClass* UConstantQNRT::GetSupportedClass() const
+{
+	return UConstantQNRT::StaticClass();
+}
+
 FName UConstantQNRT::GetAnalyzerNRTFactoryName() const 
 {
 	static const FName FactoryName(TEXT("ConstantQNRTFactory"));
diff --git a/Engine/Plugins/Runtime/AudioSynesthesia/Source/AudioSynesthesia/Private/LoudnessNRT.cpp b/Engine/Plugins/Runtime/AudioSynesthesia/Source/AudioSynesthesia/Private/LoudnessNRT.cpp
--- a/Engine/Plugins/Runtime/AudioSynesthesia/Source/AudioSynesthesia/Private/LoudnessNRT.cpp
+++ b/Engine/Plugins/Runtime/AudioSynesthesia/Source/AudioSynesthesia/Private/LoudnessNRT.cpp
@@ -72,6 +72,16 @@ TUniquePtr<Audio::IAnalyzerNRTSettings> ULoudnessNRTSettings::GetSettings(const
 	return Settings;
 }
 
+FText ULoudnessNRTSettings::GetAssetActionName() const
+{
+	return NSLOCTEXT("AssetTypeActions", "AssetTypeActions_AssetSoundSynesthesiaLoudnessNRTSettings", "Synesthesia NRT Settings (Loudness)");
+}
+
+UClass* ULoudnessNRTSettings::GetSupportedClass() const
+{
+	return ULoudnessNRTSettings::StaticClass();
+}
+
 /***************************************************************************/
 /**********************        ULoudnessNRT         ************************/
 /***************************************************************************/
@@ -162,6 +172,16 @@ TUniquePtr<Audio::IAnalyzerNRTSettings> ULoudnessNRT::GetSettings(const float In
 	return AnalyzerSettings;
 }
 
+FText ULoudnessNRT::GetAssetActionName() const
+{
+	return NSLOCTEXT("AssetTypeActions", "AssetTypeActions_AssetSoundSynesthesiaLoudnessNRT", "Synesthesia NRT (Loudness)");
+}
+
+UClass* ULoudnessNRT::GetSupportedClass() const
+{
+	return ULoudnessNRT::StaticClass();
+}
+
 FName ULoudnessNRT::GetAnalyzerNRTFactoryName() const
 {
 	static const FName FactoryName(TEXT("LoudnessNRTFactory"));
diff --git a/Engine/Plugins/Runtime/AudioSynesthesia/Source/AudioSynesthesia/Private/OnsetNRT.cpp b/Engine/Plugins/Runtime/AudioSynesthesia/Source/AudioSynesthesia/Private/OnsetNRT.cpp
--- a/Engine/Plugins/Runtime/AudioSynesthesia/Source/AudioSynesthesia/Private/OnsetNRT.cpp
+++ b/Engine/Plugins/Runtime/AudioSynesthesia/Source/AudioSynesthesia/Private/OnsetNRT.cpp
@@ -116,6 +116,15 @@ TUniquePtr<Audio::IAnalyzerNRTSettings> UOnsetNRTSettings::GetSettings(const flo
 
 }
 
+FText UOnsetNRTSettings::GetAssetActionName() const
+{
+	return NSLOCTEXT("AssetTypeActions", "AssetTypeActions_AssetSoundSynesthesiaOnsetNRTSettings", "Synesthesia NRT Settings (Onset)");
+}
+
+UClass* UOnsetNRTSettings::GetSupportedClass() const
+{
+	return UOnsetNRTSettings::StaticClass();
+}
 
 /*************************************************************************************/
 /********************************* UOnsetNRT *****************************************/
@@ -197,6 +206,16 @@ TUniquePtr<Audio::IAnalyzerNRTSettings> UOnsetNRT::GetSettings(const float InSam
 	return AnalyzerSettings;
 }
 
+FText UOnsetNRT::GetAssetActionName() const
+{
+	return NSLOCTEXT("AssetTypeActions", "AssetTypeActions_AssetSoundSynesthesiaOnset", "Synesthesia NRT (Onset)");
+}
+
+UClass* UOnsetNRT::GetSupportedClass() const
+{
+	return UOnsetNRT::StaticClass();
+}
+
 FName UOnsetNRT::GetAnalyzerNRTFactoryName() const
 {
 	static const FName FactoryName(TEXT("OnsetNRTFactory"));
diff --git a/Engine/Plugins/Runtime/AudioSynesthesia/Source/AudioSynesthesiaEditor/Private/AssetTypeActions_AudioSynesthesiaNRT.cpp b/Engine/Plugins/Runtime/AudioSynesthesia/Source/AudioSynesthesiaEditor/Private/AssetTypeActions_AudioSynesthesiaNRT.cpp
--- a/Engine/Plugins/Runtime/AudioSynesthesia/Source/AudioSynesthesiaEditor/Private/AssetTypeActions_AudioSynesthesiaNRT.cpp
+++ b/Engine/Plugins/Runtime/AudioSynesthesia/Source/AudioSynesthesiaEditor/Private/AssetTypeActions_AudioSynesthesiaNRT.cpp
@@ -9,42 +9,69 @@ FAssetTypeActions_AudioSynesthesiaNRT::FAssetTypeActions_AudioSynesthesiaNRT(UAu
 {
 }
 
+bool FAssetTypeActions_AudioSynesthesiaNRT::CanFilter()
+{
+	// If no paired synesthesia pointer provided, we filter as its a base class.
+	// Otherwise, we do not as this bloats the filter list.
+	return Synesthesia == nullptr;
+}
+
 FText FAssetTypeActions_AudioSynesthesiaNRT::GetName() const
 {
-	FText AssetActionName = Synesthesia->GetAssetActionName();
-	if (AssetActionName.IsEmpty())
-	{
-		FString ClassName;
-		Synesthesia->GetClass()->GetName(ClassName);
-		return FText::FromString(ClassName);
-	}
-	else
+	if (Synesthesia)
 	{
+		const FText AssetActionName = Synesthesia->GetAssetActionName();
+		if (AssetActionName.IsEmpty())
+		{
+			FString ClassName;
+			Synesthesia->GetClass()->GetName(ClassName);
+			return FText::FromString(ClassName);
+		}
+
 		return AssetActionName;
 	}
+
+	static const FText DefaultAssetActionName = NSLOCTEXT("AssetTypeActions", "AssetTypeActions_AssetSoundSynesthesiaNRT", "Synesthesia NRT");
+	return DefaultAssetActionName;
 }
 
 const TArray<FText>& FAssetTypeActions_AudioSynesthesiaNRT::GetSubMenus() const
 {
+	static const TArray<FText> SubMenus
+	{
+		NSLOCTEXT("AssetTypeActions", "AssetTypeActions_AssetSoundAnalysisSubmenu", "Analysis")
+	};
+
+	if (!Synesthesia)
+	{
+		return SubMenus;
+	}
+
 	return Synesthesia->GetAssetActionSubmenus();
 }
 
 FColor FAssetTypeActions_AudioSynesthesiaNRT::GetTypeColor() const 
 {
+	if (!Synesthesia)
+	{
+		return FColor(200.0f, 150.0f, 200.0f);
+	}
 	return Synesthesia->GetTypeColor(); 
 }
 
 UClass* FAssetTypeActions_AudioSynesthesiaNRT::GetSupportedClass() const
 {
-	UClass* SupportedClass = Synesthesia->GetSupportedClass();
-	if (SupportedClass == nullptr)
+	if (Synesthesia)
 	{
+		if (UClass* SupportedClass = Synesthesia->GetSupportedClass())
+		{
+			return SupportedClass;
+		}
+
 		return Synesthesia->GetClass();
 	}
-	else
-	{
-		return SupportedClass;
-	}
+
+	return UAudioSynesthesiaNRT::StaticClass();
 }
 
 uint32 FAssetTypeActions_AudioSynesthesiaNRT::GetCategories() 
diff --git a/Engine/Plugins/Runtime/AudioSynesthesia/Source/AudioSynesthesiaEditor/Private/AssetTypeActions_AudioSynesthesiaNRT.h b/Engine/Plugins/Runtime/AudioSynesthesia/Source/AudioSynesthesiaEditor/Private/AssetTypeActions_AudioSynesthesiaNRT.h
--- a/Engine/Plugins/Runtime/AudioSynesthesia/Source/AudioSynesthesiaEditor/Private/AssetTypeActions_AudioSynesthesiaNRT.h
+++ b/Engine/Plugins/Runtime/AudioSynesthesia/Source/AudioSynesthesiaEditor/Private/AssetTypeActions_AudioSynesthesiaNRT.h
@@ -14,6 +14,7 @@ public:
 	FAssetTypeActions_AudioSynesthesiaNRT(UAudioSynesthesiaNRT* InSynesthesia);
 
 	//~ Begin FAssetTypeActions_Base
+	virtual bool CanFilter() override;
 	virtual FText GetName() const override;
 	virtual FColor GetTypeColor() const override;
 	virtual UClass* GetSupportedClass() const override;
diff --git a/Engine/Plugins/Runtime/AudioSynesthesia/Source/AudioSynesthesiaEditor/Private/AssetTypeActions_AudioSynesthesiaNRTSettings.cpp b/Engine/Plugins/Runtime/AudioSynesthesia/Source/AudioSynesthesiaEditor/Private/AssetTypeActions_AudioSynesthesiaNRTSettings.cpp
--- a/Engine/Plugins/Runtime/AudioSynesthesia/Source/AudioSynesthesiaEditor/Private/AssetTypeActions_AudioSynesthesiaNRTSettings.cpp
+++ b/Engine/Plugins/Runtime/AudioSynesthesia/Source/AudioSynesthesiaEditor/Private/AssetTypeActions_AudioSynesthesiaNRTSettings.cpp
@@ -9,42 +9,69 @@ FAssetTypeActions_AudioSynesthesiaNRTSettings::FAssetTypeActions_AudioSynesthesi
 {
 }
 
+bool FAssetTypeActions_AudioSynesthesiaNRTSettings::CanFilter()
+{
+	// If no paired settings pointer provided, we filter as its a base class.
+	// Otherwise, we do not as this bloats the filter list.
+	return SynesthesiaSettings == nullptr;
+}
+
 FText FAssetTypeActions_AudioSynesthesiaNRTSettings::GetName() const
 {
-	FText AssetActionName = SynesthesiaSettings->GetAssetActionName();
-	if (AssetActionName.IsEmpty())
-	{
-		FString ClassName;
-		SynesthesiaSettings->GetClass()->GetName(ClassName);
-		return FText::FromString(ClassName);
-	}
-	else
+	if (SynesthesiaSettings)
 	{
+		const FText AssetActionName = SynesthesiaSettings->GetAssetActionName();
+		if (AssetActionName.IsEmpty())
+		{
+			FString ClassName;
+			SynesthesiaSettings->GetClass()->GetName(ClassName);
+			return FText::FromString(ClassName);
+		}
+
 		return AssetActionName;
 	}
+
+	static const FText DefaultAssetActionName = NSLOCTEXT("AssetTypeActions", "AssetTypeActions_AssetSoundSynesthesiaNRTSettings", "Synesthesia NRT Settings");
+	return DefaultAssetActionName;
 }
 
 FColor FAssetTypeActions_AudioSynesthesiaNRTSettings::GetTypeColor() const 
 {
+	if (!SynesthesiaSettings)
+	{
+		return FColor(100.0f, 100.0f, 100.0f);
+	}
 	return SynesthesiaSettings->GetTypeColor();
 }
 
 const TArray<FText>& FAssetTypeActions_AudioSynesthesiaNRTSettings::GetSubMenus() const
 {
+	static const TArray<FText> SubMenus
+	{
+		NSLOCTEXT("AssetTypeActions", "AssetTypeActions_AssetSoundAnalysisSubmenu", "Analysis")
+	};
+
+	if (!SynesthesiaSettings)
+	{
+		return SubMenus;
+	}
+
 	return SynesthesiaSettings->GetAssetActionSubmenus();
 }
 
 UClass* FAssetTypeActions_AudioSynesthesiaNRTSettings::GetSupportedClass() const
 {
-	UClass* SupportedClass = SynesthesiaSettings->GetSupportedClass();
-	if (SupportedClass == nullptr)
+	if (SynesthesiaSettings)
 	{
+		if (UClass* SupportedClass = SynesthesiaSettings->GetSupportedClass())
+		{
+			return SupportedClass;
+		}
+
 		return SynesthesiaSettings->GetClass();
 	}
-	else
-	{
-		return SupportedClass;
-	}
+
+	return UAudioSynesthesiaNRTSettings::StaticClass();
 }
 
 uint32 FAssetTypeActions_AudioSynesthesiaNRTSettings::GetCategories() 
diff --git a/Engine/Plugins/Runtime/AudioSynesthesia/Source/AudioSynesthesiaEditor/Private/AssetTypeActions_AudioSynesthesiaNRTSettings.h b/Engine/Plugins/Runtime/AudioSynesthesia/Source/AudioSynesthesiaEditor/Private/AssetTypeActions_AudioSynesthesiaNRTSettings.h
--- a/Engine/Plugins/Runtime/AudioSynesthesia/Source/AudioSynesthesiaEditor/Private/AssetTypeActions_AudioSynesthesiaNRTSettings.h
+++ b/Engine/Plugins/Runtime/AudioSynesthesia/Source/AudioSynesthesiaEditor/Private/AssetTypeActions_AudioSynesthesiaNRTSettings.h
@@ -14,6 +14,7 @@ public:
 	FAssetTypeActions_AudioSynesthesiaNRTSettings(UAudioSynesthesiaNRTSettings* InSynesthesiaSettings);
 
 	//~ Begin FAssetTypeActions_Base
+	virtual bool CanFilter() override;
 	virtual FText GetName() const override;
 	virtual FColor GetTypeColor() const override;
 	virtual UClass* GetSupportedClass() const override;
diff --git a/Engine/Plugins/Runtime/AudioSynesthesia/Source/AudioSynesthesiaEditor/Private/AudioSynesthesiaEditorModule.cpp b/Engine/Plugins/Runtime/AudioSynesthesia/Source/AudioSynesthesiaEditor/Private/AudioSynesthesiaEditorModule.cpp
--- a/Engine/Plugins/Runtime/AudioSynesthesia/Source/AudioSynesthesiaEditor/Private/AudioSynesthesiaEditorModule.cpp
+++ b/Engine/Plugins/Runtime/AudioSynesthesia/Source/AudioSynesthesiaEditor/Private/AudioSynesthesiaEditorModule.cpp
@@ -40,6 +40,9 @@ private:
 		// Register the audio editor asset type actions
 		IAssetTools& AssetTools = FModuleManager::LoadModuleChecked<FAssetToolsModule>("AssetTools").Get();
 
+		// Register base abstract class for asset filtering
+		AssetTools.RegisterAssetTypeActions(MakeShared<SynesthesiaAssetActionsType>(nullptr));
+
 		// Look for any sound effect presets to register
 		for (TObjectIterator<UClass> It; It; ++It)
 		{
@@ -54,10 +57,12 @@ private:
 			if (ParentClass->IsChildOf(SynesthesiaAssetType::StaticClass()))
 			{
 				SynesthesiaAssetType* Synesthesia = ChildClass->GetDefaultObject<SynesthesiaAssetType>();
+				check(Synesthesia);
+
 				if (!RegisteredActions.Contains(Synesthesia) && Synesthesia->HasAssetActions())
 				{
 					RegisteredActions.Add(Synesthesia);
-					AssetTools.RegisterAssetTypeActions(MakeShareable(new SynesthesiaAssetActionsType(Synesthesia)));
+					AssetTools.RegisterAssetTypeActions(MakeShared<SynesthesiaAssetActionsType>(Synesthesia));
 				}
 			}
 		}
diff --git a/Engine/Source/Runtime/AudioAnalyzer/Classes/AudioAnalyzerAsset.h b/Engine/Source/Runtime/AudioAnalyzer/Classes/AudioAnalyzerAsset.h
--- a/Engine/Source/Runtime/AudioAnalyzer/Classes/AudioAnalyzerAsset.h
+++ b/Engine/Source/Runtime/AudioAnalyzer/Classes/AudioAnalyzerAsset.h
@@ -24,12 +24,12 @@ class AUDIOANALYZER_API UAudioAnalyzerAsset : public UObject
 		/**
 		 * GetAssetActionName() returns the FText displayed in the editor.
 		 */ 
-		virtual FText GetAssetActionName() const PURE_VIRTUAL(UAudioAnalyzerNRTAsset::GetAssetActionName, return FText(););
+		virtual FText GetAssetActionName() const PURE_VIRTUAL(UAudioAnalyzerAsset::GetAssetActionName, return FText(););
 
 		/**
 		 * GetSupportedClass() returns the class which should be associated with these asset actions.
 		 */ 
-		virtual UClass* GetSupportedClass() const PURE_VIRTUAL(UAudioAnalyzerNRTAsset::GetSupportedClass, return nullptr;);
+		virtual UClass* GetSupportedClass() const PURE_VIRTUAL(UAudioAnalyzerAsset::GetSupportedClass, return nullptr;);
 
 		/**
 		 * GetTypeColor() returns the color used to display asset icons within the editor.
