Refactor audio sysnesthesia modules to have core module with no engine dependecies.

diff --git a/Engine/Plugins/Runtime/AudioSynesthesia/AudioSynesthesia.uplugin b/Engine/Plugins/Runtime/AudioSynesthesia/AudioSynesthesia.uplugin
--- a/Engine/Plugins/Runtime/AudioSynesthesia/AudioSynesthesia.uplugin
+++ b/Engine/Plugins/Runtime/AudioSynesthesia/AudioSynesthesia.uplugin
@@ -16,6 +16,11 @@
 	"Installed" : false,
 	"Modules" :
 	[
+		{
+			"Name" : "AudioSynesthesiaCore",
+			"Type" : "Runtime",
+			"LoadingPhase" : "PreDefault"
+		},
 		{
 			"Name" : "AudioSynesthesia",
 			"Type" : "Runtime",
diff --git a/Engine/Plugins/Runtime/AudioSynesthesia/Source/AudioSynesthesia/AudioSynesthesia.Build.cs b/Engine/Plugins/Runtime/AudioSynesthesia/Source/AudioSynesthesia/AudioSynesthesia.Build.cs
--- a/Engine/Plugins/Runtime/AudioSynesthesia/Source/AudioSynesthesia/AudioSynesthesia.Build.cs
+++ b/Engine/Plugins/Runtime/AudioSynesthesia/Source/AudioSynesthesia/AudioSynesthesia.Build.cs
@@ -9,8 +9,6 @@ namespace UnrealBuildTool.Rules
 	{
         public AudioSynesthesia(ReadOnlyTargetRules Target) : base(Target)
 		{
-            //bFasterWithoutUnity = true;
-			
             PublicDependencyModuleNames.AddRange(
 				new string[] {
                     "Core",
@@ -19,7 +17,8 @@ namespace UnrealBuildTool.Rules
 					"SignalProcessing",
 					"AudioMixerCore",
 					"AudioMixer",
-                    "AudioAnalyzer"
+                    "AudioAnalyzer",
+					"AudioSynesthesiaCore"
                 }
             );
 		}
diff --git a/Engine/Plugins/Runtime/AudioSynesthesia/Source/AudioSynesthesia/Private/AudioSynesthesiaCustomVersion.cpp b/Engine/Plugins/Runtime/AudioSynesthesia/Source/AudioSynesthesia/Private/AudioSynesthesiaCustomVersion.cpp
deleted file mode 100644
diff --git a/Engine/Plugins/Runtime/AudioSynesthesia/Source/AudioSynesthesia/Private/AudioSynesthesiaCustomVersion.h b/Engine/Plugins/Runtime/AudioSynesthesia/Source/AudioSynesthesia/Private/AudioSynesthesiaCustomVersion.h
deleted file mode 100644
diff --git a/Engine/Plugins/Runtime/AudioSynesthesia/Source/AudioSynesthesia/Private/AudioSynesthesiaModule.cpp b/Engine/Plugins/Runtime/AudioSynesthesia/Source/AudioSynesthesia/Private/AudioSynesthesiaModule.cpp
--- a/Engine/Plugins/Runtime/AudioSynesthesia/Source/AudioSynesthesia/Private/AudioSynesthesiaModule.cpp
+++ b/Engine/Plugins/Runtime/AudioSynesthesia/Source/AudioSynesthesia/Private/AudioSynesthesiaModule.cpp
@@ -15,17 +15,10 @@ namespace Audio
 {
 	void FAudioSynesthesiaModule::StartupModule()
 	{
-		// Register factories on startup
-		IModularFeatures::Get().RegisterModularFeature(FLoudnessNRTFactory::GetModularFeatureName(), &LoudnessFactory);
-		IModularFeatures::Get().RegisterModularFeature(FConstantQNRTFactory::GetModularFeatureName(), &ConstantQFactory);
-		IModularFeatures::Get().RegisterModularFeature(FOnsetNRTFactory::GetModularFeatureName(), &OnsetFactory);
+		FModuleManager::Get().LoadModuleChecked<IModuleInterface>("AudioSynesthesiaCore");
 	}
 
 	void FAudioSynesthesiaModule::ShutdownModule()
 	{
-		// Unregister factories on shutdown
-		IModularFeatures::Get().UnregisterModularFeature(FLoudnessNRTFactory::GetModularFeatureName(), &LoudnessFactory);
-		IModularFeatures::Get().UnregisterModularFeature(FConstantQNRTFactory::GetModularFeatureName(), &ConstantQFactory);
-		IModularFeatures::Get().UnregisterModularFeature(FOnsetNRTFactory::GetModularFeatureName(), &OnsetFactory);
 	}
 }
diff --git a/Engine/Plugins/Runtime/AudioSynesthesia/Source/AudioSynesthesia/Private/ConstantQAnalyzer.cpp b/Engine/Plugins/Runtime/AudioSynesthesia/Source/AudioSynesthesia/Private/ConstantQAnalyzer.cpp
deleted file mode 100644
diff --git a/Engine/Plugins/Runtime/AudioSynesthesia/Source/AudioSynesthesia/Private/ConstantQNRTFactory.cpp b/Engine/Plugins/Runtime/AudioSynesthesia/Source/AudioSynesthesia/Private/ConstantQNRTFactory.cpp
deleted file mode 100644
diff --git a/Engine/Plugins/Runtime/AudioSynesthesia/Source/AudioSynesthesia/Private/LoudnessAnalyzer.cpp b/Engine/Plugins/Runtime/AudioSynesthesia/Source/AudioSynesthesia/Private/LoudnessAnalyzer.cpp
deleted file mode 100644
diff --git a/Engine/Plugins/Runtime/AudioSynesthesia/Source/AudioSynesthesia/Private/LoudnessNRTFactory.cpp b/Engine/Plugins/Runtime/AudioSynesthesia/Source/AudioSynesthesia/Private/LoudnessNRTFactory.cpp
deleted file mode 100644
diff --git a/Engine/Plugins/Runtime/AudioSynesthesia/Source/AudioSynesthesia/Private/OnsetAnalyzer.cpp b/Engine/Plugins/Runtime/AudioSynesthesia/Source/AudioSynesthesia/Private/OnsetAnalyzer.cpp
deleted file mode 100644
diff --git a/Engine/Plugins/Runtime/AudioSynesthesia/Source/AudioSynesthesia/Private/OnsetNRTFactory.cpp b/Engine/Plugins/Runtime/AudioSynesthesia/Source/AudioSynesthesia/Private/OnsetNRTFactory.cpp
deleted file mode 100644
diff --git a/Engine/Plugins/Runtime/AudioSynesthesia/Source/AudioSynesthesia/Private/PeakPicker.cpp b/Engine/Plugins/Runtime/AudioSynesthesia/Source/AudioSynesthesia/Private/PeakPicker.cpp
deleted file mode 100644
diff --git a/Engine/Plugins/Runtime/AudioSynesthesia/Source/AudioSynesthesia/Public/ConstantQAnalyzer.h b/Engine/Plugins/Runtime/AudioSynesthesia/Source/AudioSynesthesia/Public/ConstantQAnalyzer.h
deleted file mode 100644
diff --git a/Engine/Plugins/Runtime/AudioSynesthesia/Source/AudioSynesthesia/Public/ConstantQNRTFactory.h b/Engine/Plugins/Runtime/AudioSynesthesia/Source/AudioSynesthesia/Public/ConstantQNRTFactory.h
deleted file mode 100644
diff --git a/Engine/Plugins/Runtime/AudioSynesthesia/Source/AudioSynesthesia/Public/InterpolateSorted.h b/Engine/Plugins/Runtime/AudioSynesthesia/Source/AudioSynesthesia/Public/InterpolateSorted.h
deleted file mode 100644
diff --git a/Engine/Plugins/Runtime/AudioSynesthesia/Source/AudioSynesthesia/Public/LoudnessAnalyzer.h b/Engine/Plugins/Runtime/AudioSynesthesia/Source/AudioSynesthesia/Public/LoudnessAnalyzer.h
deleted file mode 100644
diff --git a/Engine/Plugins/Runtime/AudioSynesthesia/Source/AudioSynesthesia/Public/LoudnessNRTFactory.h b/Engine/Plugins/Runtime/AudioSynesthesia/Source/AudioSynesthesia/Public/LoudnessNRTFactory.h
deleted file mode 100644
diff --git a/Engine/Plugins/Runtime/AudioSynesthesia/Source/AudioSynesthesia/Public/OnsetAnalyzer.h b/Engine/Plugins/Runtime/AudioSynesthesia/Source/AudioSynesthesia/Public/OnsetAnalyzer.h
deleted file mode 100644
diff --git a/Engine/Plugins/Runtime/AudioSynesthesia/Source/AudioSynesthesia/Public/OnsetNRTFactory.h b/Engine/Plugins/Runtime/AudioSynesthesia/Source/AudioSynesthesia/Public/OnsetNRTFactory.h
deleted file mode 100644
diff --git a/Engine/Plugins/Runtime/AudioSynesthesia/Source/AudioSynesthesia/Public/PeakPicker.h b/Engine/Plugins/Runtime/AudioSynesthesia/Source/AudioSynesthesia/Public/PeakPicker.h
deleted file mode 100644
diff --git a/Engine/Plugins/Runtime/AudioSynesthesia/Source/AudioSynesthesiaCore/AudioSynesthesiaCore.Build.cs b/Engine/Plugins/Runtime/AudioSynesthesia/Source/AudioSynesthesiaCore/AudioSynesthesiaCore.Build.cs
new file mode 100644
--- /dev/null
+++ b/Engine/Plugins/Runtime/AudioSynesthesia/Source/AudioSynesthesiaCore/AudioSynesthesiaCore.Build.cs
@@ -0,0 +1,32 @@
+// Copyright Epic Games, Inc. All Rights Reserved.
+
+using UnrealBuildTool;
+using System.IO;
+
+namespace UnrealBuildTool.Rules
+{
+	public class AudioSynesthesiaCore : ModuleRules
+	{
+        public AudioSynesthesiaCore(ReadOnlyTargetRules Target) : base(Target)
+		{
+			PublicIncludePathModuleNames.AddRange(
+				new string[] {
+					"AudioAnalyzer"
+				}
+			);
+			
+            PublicDependencyModuleNames.AddRange(
+				new string[] {
+                    "Core",
+					"SignalProcessing",
+                }
+            );
+
+			PrivateIncludePathModuleNames.AddRange(
+				new string[] {
+					"AudioMixer"
+				}
+			);
+		}
+	}
+}
diff --git a/Engine/Plugins/Runtime/AudioSynesthesia/Source/AudioSynesthesiaCore/Private/AudioSynesthesiaCoreModule.cpp b/Engine/Plugins/Runtime/AudioSynesthesia/Source/AudioSynesthesiaCore/Private/AudioSynesthesiaCoreModule.cpp
new file mode 100644
--- /dev/null
+++ b/Engine/Plugins/Runtime/AudioSynesthesia/Source/AudioSynesthesiaCore/Private/AudioSynesthesiaCoreModule.cpp
@@ -0,0 +1,43 @@
+// Copyright Epic Games, Inc. All Rights Reserved.
+
+#include "AudioSynesthesiaCoreModule.h"
+#include "CoreMinimal.h"
+#include "Modules/ModuleInterface.h"
+#include "Modules/ModuleManager.h"
+#include "Features/IModularFeatures.h"
+#include "LoudnessNRTFactory.h"
+#include "ConstantQNRTFactory.h"
+#include "OnsetNRTFactory.h"
+
+
+namespace Audio
+{
+	class FAudioSynesthesiaCoreModule : public IAudioSynesthesiaCoreModule
+	{
+		public:
+			void StartupModule()
+			{
+				// Register factories on startup
+				IModularFeatures::Get().RegisterModularFeature(FLoudnessNRTFactory::GetModularFeatureName(), &LoudnessFactory);
+				IModularFeatures::Get().RegisterModularFeature(FConstantQNRTFactory::GetModularFeatureName(), &ConstantQFactory);
+				IModularFeatures::Get().RegisterModularFeature(FOnsetNRTFactory::GetModularFeatureName(), &OnsetFactory);
+			}
+
+			void ShutdownModule()
+			{
+				// Unregister factories on shutdown
+				IModularFeatures::Get().UnregisterModularFeature(FLoudnessNRTFactory::GetModularFeatureName(), &LoudnessFactory);
+				IModularFeatures::Get().UnregisterModularFeature(FConstantQNRTFactory::GetModularFeatureName(), &ConstantQFactory);
+				IModularFeatures::Get().UnregisterModularFeature(FOnsetNRTFactory::GetModularFeatureName(), &OnsetFactory);
+			}
+
+		private:
+			FLoudnessNRTFactory LoudnessFactory;
+			FConstantQNRTFactory ConstantQFactory;
+			FOnsetNRTFactory OnsetFactory;
+	};
+
+}
+
+IMPLEMENT_MODULE(Audio::FAudioSynesthesiaCoreModule, AudioSynesthesiaCore);
+
diff --git a/Engine/Plugins/Runtime/AudioSynesthesia/Source/AudioSynesthesiaCore/Private/AudioSynesthesiaCustomVersion.cpp b/Engine/Plugins/Runtime/AudioSynesthesia/Source/AudioSynesthesiaCore/Private/AudioSynesthesiaCustomVersion.cpp
new file mode 100644
--- /dev/null
+++ b/Engine/Plugins/Runtime/AudioSynesthesia/Source/AudioSynesthesiaCore/Private/AudioSynesthesiaCustomVersion.cpp
@@ -0,0 +1,9 @@
+// Copyright Epic Games, Inc. All Rights Reserved.
+
+#include "AudioSynesthesiaCustomVersion.h"
+#include "Serialization/CustomVersion.h"
+
+const FGuid FAudioSynesthesiaCustomVersion::GUID( 0x82E77C4E, 0x332343A5, 0xB46B13C5, 0x97310DF3 );
+
+// Register the custom version with core
+FCustomVersionRegistration GRegisterAudioSynesthesiaCustomVersion( FAudioSynesthesiaCustomVersion::GUID, FAudioSynesthesiaCustomVersion::LatestVersion, TEXT( "AudioSynesthesiaVer" ) );
diff --git a/Engine/Plugins/Runtime/AudioSynesthesia/Source/AudioSynesthesiaCore/Private/AudioSynesthesiaCustomVersion.h b/Engine/Plugins/Runtime/AudioSynesthesia/Source/AudioSynesthesiaCore/Private/AudioSynesthesiaCustomVersion.h
new file mode 100644
--- /dev/null
+++ b/Engine/Plugins/Runtime/AudioSynesthesia/Source/AudioSynesthesiaCore/Private/AudioSynesthesiaCustomVersion.h
@@ -0,0 +1,25 @@
+// Copyright Epic Games, Inc. All Rights Reserved.
+
+#pragma once
+
+#include "CoreMinimal.h"
+#include "Misc/Guid.h"
+
+// Custom serialization version for all packages containing AudioSynesthesia asset types
+struct FAudioSynesthesiaCustomVersion
+{
+	enum Type
+	{
+		// Before any version changes were made in the module
+		BeforeCustomVersionWasAdded = 0,
+
+		// -----<new versions can be added above this line>-------------------------------------------------
+		VersionPlusOne,
+		LatestVersion = VersionPlusOne - 1
+	};
+
+	// The GUID for this custom version number
+	const static FGuid GUID;
+
+	FAudioSynesthesiaCustomVersion() = delete;
+};
diff --git a/Engine/Plugins/Runtime/AudioSynesthesia/Source/AudioSynesthesiaCore/Private/ConstantQAnalyzer.cpp b/Engine/Plugins/Runtime/AudioSynesthesia/Source/AudioSynesthesiaCore/Private/ConstantQAnalyzer.cpp
new file mode 100644
--- /dev/null
+++ b/Engine/Plugins/Runtime/AudioSynesthesia/Source/AudioSynesthesiaCore/Private/ConstantQAnalyzer.cpp
@@ -0,0 +1,83 @@
+// Copyright Epic Games, Inc. All Rights Reserved.
+
+#include "ConstantQAnalyzer.h"
+
+#include "DSP/ConstantQ.h"
+#include "DSP/AudioFFT.h"
+#include "DSP/FloatArrayMath.h"
+#include "DSP/BufferVectorOperations.h"
+
+namespace Audio
+{
+	FConstantQAnalyzer::FConstantQAnalyzer(const FConstantQAnalyzerSettings& InSettings, const float InSampleRate)
+	:	Settings(InSettings)
+	,	SampleRate(InSampleRate)
+	,	NumUsefulFFTBins((Settings.FFTSize / 2) + 1)
+    ,   Window(Settings.WindowType, Settings.FFTSize, 1, false)
+	{
+		// Need FFTSize atleast 8 to support optimized operations.
+		check(Settings.FFTSize >= 8);
+		check(SampleRate > 0.f);
+
+		CQTTransform = NewPseudoConstantQKernelTransform(Settings, Settings.FFTSize, SampleRate);
+
+		// Resize internal buffers
+		WindowedSamples.Reset(Settings.FFTSize);
+		WindowedSamples.AddUninitialized(Settings.FFTSize);
+
+		FFTOutputRealData.Reset(Settings.FFTSize);
+		FFTOutputRealData.AddUninitialized(Settings.FFTSize);
+
+		FFTOutputImagData.Reset(Settings.FFTSize);
+		FFTOutputImagData.AddUninitialized(Settings.FFTSize);
+
+		SpectrumBuffer.Reset(NumUsefulFFTBins);
+		SpectrumBuffer.AddUninitialized(NumUsefulFFTBins);
+	}
+
+	void FConstantQAnalyzer::CalculateCQT(const float* InSamples, TArray<float>& OutCQT)
+	{
+		// Copy input samples and apply window
+		FMemory::Memcpy(WindowedSamples.GetData(), InSamples, sizeof(float) * Settings.FFTSize);
+		Window.ApplyToBuffer(WindowedSamples.GetData());
+
+		// Take FFT of windowed samples
+		const FFTTimeDomainData TimeData = {WindowedSamples.GetData(), Settings.FFTSize};
+		FFTFreqDomainData FreqData = {FFTOutputRealData.GetData(), FFTOutputImagData.GetData()};
+
+		PerformFFT(TimeData, FreqData);
+
+		// Calculate spectrum from FFT output
+		ComputeSpectrum(Settings.SpectrumType, FreqData, Settings.FFTSize, SpectrumBuffer);
+		
+		// Convert spectrum to CQT
+		CQTTransform->TransformArray(SpectrumBuffer, OutCQT);
+
+		// Apply decibel scaling if appropriate.
+		if (EConstantQScaling::Decibel == Settings.Scaling)
+		{
+			// Convert to dB
+			switch (Settings.SpectrumType)
+			{
+				case ESpectrumType::MagnitudeSpectrum:
+					ArrayMagnitudeToDecibelInPlace(OutCQT, -90.f);
+					break;
+
+				case ESpectrumType::PowerSpectrum:
+					ArrayPowerToDecibelInPlace(OutCQT, -90.f);
+					break;
+
+				default:
+					check(false);
+					//checkf(false, TEXT("Unhandled ESpectrumType %s"), GETENUMSTRING(ESpectrumType, Settings.SpectrumType));
+					ArrayPowerToDecibelInPlace(OutCQT, -90.f);
+			}
+		}
+	}
+
+	const FConstantQAnalyzerSettings& FConstantQAnalyzer::GetSettings() const
+	{
+		return Settings;
+	}
+}
+
diff --git a/Engine/Plugins/Runtime/AudioSynesthesia/Source/AudioSynesthesiaCore/Private/ConstantQNRTFactory.cpp b/Engine/Plugins/Runtime/AudioSynesthesia/Source/AudioSynesthesiaCore/Private/ConstantQNRTFactory.cpp
new file mode 100644
--- /dev/null
+++ b/Engine/Plugins/Runtime/AudioSynesthesia/Source/AudioSynesthesiaCore/Private/ConstantQNRTFactory.cpp
@@ -0,0 +1,272 @@
+// Copyright Epic Games, Inc. All Rights Reserved.
+
+#include "ConstantQNRTFactory.h"
+#include "AudioSynesthesiaCustomVersion.h"
+#include "DSP/ConstantQ.h"
+#include "DSP/DeinterleaveView.h"
+#include "DSP/FloatArrayMath.h"
+
+namespace Audio 
+{
+	/*******************************************************************************/
+	/*************************** FConstantQNRTResult *******************************/
+	/*******************************************************************************/
+
+	FArchive &operator <<(FArchive& Ar, FConstantQFrame& Frame)
+	{
+		Ar.UsingCustomVersion(FAudioSynesthesiaCustomVersion::GUID);
+
+		Ar << Frame.Channel;
+		Ar << Frame.Timestamp;
+		Ar << Frame.Spectrum;
+
+		return Ar;
+	}
+
+	FConstantQNRTResult::FConstantQNRTResult()
+	:	DurationInSeconds(0.f)
+	,	bIsSortedChronologically(false)
+	{}
+
+	void FConstantQNRTResult::Serialize(FArchive& Archive)
+	{
+		Archive.UsingCustomVersion(FAudioSynesthesiaCustomVersion::GUID);
+		
+		Archive << DurationInSeconds;
+		Archive << bIsSortedChronologically;
+		Archive << ChannelCQTFrames;
+		Archive << ChannelCQTIntervals;
+	}
+
+	void FConstantQNRTResult::AddFrame(int32 InChannelIndex, float InTimestamp, TArrayView<const float> InSpectrum)
+	{
+		TArray<FConstantQFrame>& FrameArray = ChannelCQTFrames.FindOrAdd(InChannelIndex);
+		FrameArray.Emplace(InChannelIndex, InTimestamp, InSpectrum);
+
+		FFloatInterval& Interval = ChannelCQTIntervals.FindOrAdd(InChannelIndex);
+		for (float Value : InSpectrum)
+		{
+			Interval.Include(Value);
+		}
+
+		bIsSortedChronologically = false;
+	}
+
+	bool FConstantQNRTResult::ContainsChannel(int32 InChannelIndex) const
+	{
+		return ChannelCQTFrames.Contains(InChannelIndex);
+	}
+
+	const TArray<FConstantQFrame>& FConstantQNRTResult::GetFramesForChannel(int32 InChannelIndex) const
+	{
+		return ChannelCQTFrames[InChannelIndex];
+	}
+
+	FFloatInterval FConstantQNRTResult::GetChannelConstantQInterval(int32 InChannelIdx) const
+	{
+		return ChannelCQTIntervals[InChannelIdx];
+	}
+
+	void FConstantQNRTResult::GetChannels(TArray<int32>& OutChannels) const
+	{
+		ChannelCQTFrames.GetKeys(OutChannels);
+	}
+
+	float FConstantQNRTResult::GetDurationInSeconds() const 
+	{
+		return DurationInSeconds;
+	}
+
+	void FConstantQNRTResult::SetDurationInSeconds(float InDuration)
+	{
+		DurationInSeconds = InDuration;
+	}
+
+	 /* Returns true if FConstantQNRTFrame arrays are sorted in chronologically ascending order via their timestamp.
+	 */
+	bool FConstantQNRTResult::IsSortedChronologically() const
+	{
+		return bIsSortedChronologically;
+	}
+
+	/**
+	 * Sorts FConstantQNRTFrame arrays in chronologically ascnding order via their timestamp.
+	 */
+	void FConstantQNRTResult::SortChronologically()
+	{
+		for (auto& KeyValue : ChannelCQTFrames)
+		{
+			KeyValue.Value.Sort([](const FConstantQFrame& A, const FConstantQFrame& B) { return A.Timestamp < B.Timestamp; });
+		}
+
+		bIsSortedChronologically = true;
+	}
+
+	/*******************************************************************************/
+	/*************************** FConstantQNRTWorker *******************************/
+	/*******************************************************************************/
+
+	FConstantQNRTWorker::FConstantQNRTWorker(const FAnalyzerNRTParameters& InParams, const FConstantQNRTSettings& InAnalyzerSettings)
+	:	NumFrames(0)
+	,	NumChannels(InParams.NumChannels)
+	,	NumBuffers(0)
+	,	SampleRate(InParams.SampleRate)
+	,	NumHopFrames(0)
+	,	NumHopSamples(0)
+	,	NumWindowFrames(InAnalyzerSettings.FFTSize)
+	,	NumWindowSamples(InAnalyzerSettings.FFTSize * InParams.NumChannels)
+	,	MonoScaling(1.f)
+	,	bDownmixToMono(InAnalyzerSettings.bDownmixToMono)
+	{
+		check(InParams.NumChannels > 0);
+		check(InParams.SampleRate > 0.f);
+
+		check(InAnalyzerSettings.AnalysisPeriod > 0.f);
+		check(FMath::IsPowerOfTwo(InAnalyzerSettings.FFTSize));
+
+
+		NumHopFrames = FMath::Max(1, FMath::RoundToInt(InAnalyzerSettings.AnalysisPeriod * InParams.SampleRate));
+		NumHopSamples = NumHopFrames * NumChannels;
+
+		if (bDownmixToMono)
+		{
+			MonoBuffer.Reset(NumWindowFrames);
+			MonoBuffer.AddUninitialized(NumWindowFrames);
+			// equal power mixing
+			MonoScaling = 1.f / FMath::Sqrt(static_cast<float>(NumChannels));
+		}
+
+		SlidingBuffer = MakeUnique<TSlidingBuffer<float> >(NumWindowSamples, NumHopSamples);
+
+		ConstantQAnalyzer = MakeUnique<FConstantQAnalyzer>(InAnalyzerSettings, SampleRate);
+	}
+
+
+	void FConstantQNRTWorker::Analyze(TArrayView<const float> InAudio, IAnalyzerNRTResult* OutResult)
+	{
+		bool bDoFlush = false;
+
+		AnalyzeMultichannel(InAudio, OutResult, bDoFlush);
+	}
+
+	void FConstantQNRTWorker::AnalyzeMultichannel(TArrayView<const float> InAudio, IAnalyzerNRTResult* OutResult, bool bDoFlush)
+	{
+		NumFrames += InAudio.Num() / NumChannels;
+
+		// Assume that outer layers ensured that this is of correct type.
+		FConstantQNRTResult* ConstantQResult = static_cast<FConstantQNRTResult*>(OutResult);
+
+		check(nullptr != ConstantQResult);
+
+		TAutoSlidingWindow<float, FAudioBufferAlignedAllocator> SlidingWindow(*SlidingBuffer, InAudio, HopBuffer, bDoFlush);
+
+		if (bDownmixToMono)
+		{
+			for (const AlignedFloatBuffer& Window : SlidingWindow)
+			{
+				FMemory::Memset(MonoBuffer.GetData(), 0, sizeof(float) * NumWindowFrames);
+
+				// Sum all channels into mono buffer
+				TAutoDeinterleaveView<float, FAudioBufferAlignedAllocator> DeinterleaveView(Window, ChannelBuffer, NumChannels);
+				for(auto Channel : DeinterleaveView)
+				{
+					MixInBufferFast(Channel.Values, MonoBuffer);		
+				}
+
+				// Equal power sum. assuming incoherrent signals.
+				MultiplyBufferByConstantInPlace(MonoBuffer, 1.f / FMath::Sqrt(static_cast<float>(NumChannels)));
+
+				AnalyzeWindow(MonoBuffer, 0, *ConstantQResult);
+
+				NumBuffers++;
+			}
+		}
+		else
+		{
+			// Loop through entire array of input samples.
+			for (const AlignedFloatBuffer& Window : SlidingWindow)
+			{
+				// Each channel is analyzed seperately.
+				for(const TDeinterleaveView<float>::TChannel<FAudioBufferAlignedAllocator>& Channel : TAutoDeinterleaveView<float, FAudioBufferAlignedAllocator>(Window, ChannelBuffer, NumChannels))
+				{
+					AnalyzeWindow(Channel.Values, Channel.ChannelIndex, *ConstantQResult);
+				}
+
+				NumBuffers++;
+			}
+		}
+	}
+
+	/** Called when analysis of audio asset is complete. */
+	void FConstantQNRTWorker::Finalize(IAnalyzerNRTResult* OutResult) 
+	{
+		AlignedFloatBuffer EmptyArray;
+		bool bDoFlush = true;
+
+		AnalyzeMultichannel(EmptyArray, OutResult, bDoFlush);
+
+		// Assume that outer layers ensured that this is of correct type.
+		FConstantQNRTResult* ConstantQResult = static_cast<FConstantQNRTResult*>(OutResult);
+
+		check(nullptr != ConstantQResult);
+
+
+		if (nullptr != OutResult)
+		{
+			check(SampleRate > 0.f);
+
+			float DurationInSeconds = static_cast<float>(NumFrames) / SampleRate;
+
+			ConstantQResult->SetDurationInSeconds(DurationInSeconds);
+
+			ConstantQResult->SortChronologically();
+		}
+
+		// Reset internal counters
+		NumBuffers = 0; 
+		NumFrames = 0;
+
+		SlidingBuffer.Reset();
+	}
+
+	void FConstantQNRTWorker::AnalyzeWindow(const AlignedFloatBuffer& InWindow, int32 InChannelIndex, FConstantQNRTResult& OutResult)
+	{
+		// Run CQT Analyzer
+		ConstantQAnalyzer->CalculateCQT(InWindow.GetData(), CQTSpectrum);
+
+		// Get timestamp as center of audio window.
+		float Timestamp = ((NumBuffers * NumHopFrames) + (0.5f * ConstantQAnalyzer->GetSettings().FFTSize)) / SampleRate;
+
+		OutResult.AddFrame(InChannelIndex, Timestamp, CQTSpectrum);
+	}
+
+	// Name of this analyzer type.
+	FName FConstantQNRTFactory::GetName() const
+	{
+		static FName FactoryName(TEXT("ConstantQNRTFactory"));
+		return FactoryName;
+	}
+
+	// Human readable name of this analyzer.
+	FString FConstantQNRTFactory::GetTitle() const
+	{
+		return TEXT("Constant Q Analyzer Non-Real-Time");
+	}
+
+	// Create a new FConstantQNRTResult.
+	TUniquePtr<IAnalyzerNRTResult> FConstantQNRTFactory::NewResult() const
+	{
+		return MakeUnique<FConstantQNRTResult>();
+	}
+
+	// Create a new FConstantQNRTWorker
+	TUniquePtr<IAnalyzerNRTWorker> FConstantQNRTFactory::NewWorker(const FAnalyzerNRTParameters& InParams, const IAnalyzerNRTSettings* InSettings) const
+	{
+		const FConstantQNRTSettings* ConstantQSettings = static_cast<const FConstantQNRTSettings*>(InSettings);
+
+		check(nullptr != ConstantQSettings);
+
+		return MakeUnique<FConstantQNRTWorker>(InParams, *ConstantQSettings);
+	}
+}
+
diff --git a/Engine/Plugins/Runtime/AudioSynesthesia/Source/AudioSynesthesiaCore/Private/LoudnessAnalyzer.cpp b/Engine/Plugins/Runtime/AudioSynesthesia/Source/AudioSynesthesiaCore/Private/LoudnessAnalyzer.cpp
new file mode 100644
--- /dev/null
+++ b/Engine/Plugins/Runtime/AudioSynesthesia/Source/AudioSynesthesiaCore/Private/LoudnessAnalyzer.cpp
@@ -0,0 +1,271 @@
+// Copyright Epic Games, Inc. All Rights Reserved.
+
+
+#include "LoudnessAnalyzer.h"
+#include "AudioMixer.h"
+#include "DSP/BufferVectorOperations.h"
+#include "DSP/DeinterleaveView.h"
+
+namespace 
+{
+	static const float Log10Scale = 1.0f / FMath::Loge(10.f);
+}
+
+namespace Audio
+{
+    /* Equal loudness curves from: ANSI Standards S1.4-1983 and S1.42-2001 */
+    /* https://en.wikipedia.org/wiki/A-weighting */
+    float GetEqualLoudnessAWeightForFrequency(const float Freq)
+    {
+        const float FreqSqrd = Freq * Freq;
+
+        return 1.258925411794167f * ((148840000.f * FreqSqrd * FreqSqrd) / ((FreqSqrd + 424.36f) * (FreqSqrd + 148840000.f) * FMath::Sqrt((FreqSqrd + 11599.29f) * (FreqSqrd + 544496.41f))));
+    }
+
+    float GetEqualLoudnessBWeightForFrequency(const float Freq)
+    {
+        const float FreqSqrd = Freq * Freq;
+
+        return 1.019764760044717f * ((148840000.f * Freq * FreqSqrd) / ((FreqSqrd + 424.36f) * (FreqSqrd + 148840000.f) * FMath::Sqrt(FreqSqrd + 25122.25f)));
+    }
+
+    float GetEqualLoudnessCWeightForFrequency(const float Freq)
+    {
+        const float FreqSqrd = Freq * Freq;
+
+        return 1.006931668851804f * ((148840000.f * FreqSqrd) / ((FreqSqrd + 424.36f) * (FreqSqrd + 148840000.f)));
+    }
+
+    float GetEqualLoudnessDWeightForFrequency(const float Freq)
+    {
+        const float FreqSqrd = Freq * Freq;
+        const float H1 = (1037918.48f - FreqSqrd);
+        const float H2 = (9837328.f - FreqSqrd);
+        const float H = (H1 + 1080768.16f * FreqSqrd) / (H2 + 11723776.f * FreqSqrd);
+		return (Freq / 6.8966888496476e-5f) * FMath::Sqrt(H / ((FreqSqrd + 79919.29f) * (FreqSqrd + 1345600.f)));
+    }
+
+    float GetEqualLoudnessNoneWeightForFrequency(const float Freq)
+    {
+        return 1.0f;
+    }
+    
+
+	/****************************************************************************/
+	/***********************    FLoudnessAnalyzer    ****************************/
+	/****************************************************************************/
+
+    FLoudnessAnalyzer::FLoudnessAnalyzer(float InSampleRate, const FLoudnessAnalyzerSettings& InSettings)
+    :   Settings(InSettings)
+    ,   EnergyScale(1.0f)
+    ,   MinFreqIndex(0)
+    ,   MaxFreqIndex(0)
+    ,   FFTFreqSize((Settings.FFTSize / 2) + 1)
+    ,   Window(Settings.WindowType, Settings.FFTSize, 1, false)
+    {
+		checkf(FMath::IsPowerOfTwo(Settings.FFTSize), TEXT("FFT size must be a power of two"));
+		checkf(Settings.FFTSize >= 4, TEXT("FFT size must be atleast 4"))
+        check(InSampleRate > 0);
+
+        // Determine which fft bins will be used for calculating loudness
+        const float FreqPerBin = InSampleRate / Settings.FFTSize;
+        MinFreqIndex = static_cast<int32>(FMath::RoundHalfFromZero(Settings.MinAnalysisFrequency / FreqPerBin));
+        MinFreqIndex = FMath::Max(0, MinFreqIndex);
+        MaxFreqIndex = static_cast<int32>(FMath::RoundHalfFromZero(Settings.MaxAnalysisFrequency / FreqPerBin)) + 1;
+        MaxFreqIndex = FMath::Min(MaxFreqIndex, FFTFreqSize);
+
+        // Resize internal buffers
+        WindowedSamples.Reset(Settings.FFTSize);
+        WindowedSamples.AddUninitialized(Settings.FFTSize);
+
+        FFTOutputRealData.Reset(Settings.FFTSize);
+        FFTOutputRealData.AddUninitialized(Settings.FFTSize);
+
+        FFTOutputImagData.Reset(Settings.FFTSize);
+        FFTOutputImagData.AddUninitialized(Settings.FFTSize);
+
+        CurveWeights.Reset(FFTFreqSize);
+        CurveWeights.AddUninitialized(FFTFreqSize);
+
+        PowerSpectrumBuffer.Reset(FFTFreqSize);
+        PowerSpectrumBuffer.AddUninitialized(FFTFreqSize);
+
+        float(*LoudnessMapFn)(const float);
+        switch (Settings.LoudnessCurveType)
+        {
+            case ELoudnessCurveType::A:
+                LoudnessMapFn = GetEqualLoudnessAWeightForFrequency;
+                break;
+
+            case ELoudnessCurveType::B:
+                LoudnessMapFn = GetEqualLoudnessBWeightForFrequency;
+                break;
+
+            case ELoudnessCurveType::C:
+                LoudnessMapFn = GetEqualLoudnessCWeightForFrequency;
+                break;
+
+            case ELoudnessCurveType::D:
+                LoudnessMapFn = GetEqualLoudnessDWeightForFrequency;
+                break;
+
+            case ELoudnessCurveType::None:
+            default:
+                LoudnessMapFn = GetEqualLoudnessNoneWeightForFrequency;
+                break;
+        }
+
+        /* Calculate equal loudness weighting curve */
+
+        // Convention is to normalize 1kHz to have a loudness of 1.
+        float EqualLoudnessNorm = 1.0f / LoudnessMapFn(1000.0f);
+
+        float* CurveWeightsPtr = CurveWeights.GetData();
+
+        const float FFTSize = static_cast<float>(Settings.FFTSize);
+
+        for (int32 i = 0; i < FFTFreqSize; i++)
+        {
+            const float Freq = static_cast<float>(i) * InSampleRate / FFTSize;
+			float Weighting = EqualLoudnessNorm;
+
+			switch (Settings.LoudnessCurveType)
+			{
+				case ELoudnessCurveType::A:
+					Weighting *= GetEqualLoudnessAWeightForFrequency(Freq);
+					break;
+
+				case ELoudnessCurveType::B:
+					Weighting *= GetEqualLoudnessBWeightForFrequency(Freq);
+					break;
+
+				case ELoudnessCurveType::C:
+					Weighting *= GetEqualLoudnessCWeightForFrequency(Freq);
+					break;
+
+				case ELoudnessCurveType::D:
+					Weighting *= GetEqualLoudnessDWeightForFrequency(Freq);
+					break;
+
+				case ELoudnessCurveType::None:
+				default:
+					Weighting *= GetEqualLoudnessNoneWeightForFrequency(Freq);
+					break;
+			}
+
+            // Curve designed for magnitude domain, but applied to power spectrum. The curve is squared to be applied to power spectrum.
+            CurveWeightsPtr[i] = Weighting * Weighting;
+        }
+
+		// Normalize by FFT Scaling (1 / FFTSize) and by number of samples (1 / WindowSize) to calculate mean squared value.
+		EnergyScale = 1.f / static_cast<float>(Settings.FFTSize * Settings.FFTSize);
+    }
+
+    const FLoudnessAnalyzerSettings& FLoudnessAnalyzer::GetSettings() const
+    {
+        return Settings;
+    }
+
+    float FLoudnessAnalyzer::CalculatePerceptualEnergy(TArrayView<const float> InView)
+    {
+		check(InView.Num() == Settings.FFTSize);
+
+		// Copy input samples and apply window
+        FMemory::Memcpy(WindowedSamples.GetData(), InView.GetData(), sizeof(float) * Settings.FFTSize);
+        Window.ApplyToBuffer(WindowedSamples.GetData());
+
+		// Take FFT of windowed samples
+        const FFTTimeDomainData TimeData = {WindowedSamples.GetData(), Settings.FFTSize};
+        FFTFreqDomainData FreqData = {FFTOutputRealData.GetData(), FFTOutputImagData.GetData()};
+
+        PerformFFT(TimeData, FreqData);
+
+		// Accumulate energy as weighted sum of power spectrum in frequency analysis range.
+        const float* FFTOutputRealPtr = FFTOutputRealData.GetData();
+        const float* FFTOutputImagPtr = FFTOutputImagData.GetData();
+        float* PowerSpectrumPtr = PowerSpectrumBuffer.GetData();
+		const float* CurveWeightsPtr = CurveWeights.GetData();
+
+		float Total = 0.0f;
+        for (int32 i = MinFreqIndex; i < MaxFreqIndex; i++)
+        {
+            PowerSpectrumPtr[i] = FFTOutputRealPtr[i] * FFTOutputRealPtr[i] + FFTOutputImagPtr[i] * FFTOutputImagPtr[i];
+			Total += PowerSpectrumPtr[i] * CurveWeightsPtr[i];
+        }
+
+		// Normalize by FFT Scaling and by number of samples to calculate mean squared value.
+		Total *= EnergyScale;
+
+        return Total;
+    }
+
+    float FLoudnessAnalyzer::CalculateLoudness(TArrayView<const float> InView)
+    {
+        float Energy = CalculatePerceptualEnergy(InView);
+		return ConvertPerceptualEnergyToLoudness(Energy);
+    }
+
+	float FLoudnessAnalyzer::ConvertPerceptualEnergyToLoudness(float InPerceptualEnergy)
+	{
+        return 10.0f * FMath::Loge(InPerceptualEnergy) * Log10Scale;
+	}
+
+	/****************************************************************************/
+	/****************    FMultichannelLoudnessAnalyzer    ***********************/
+	/****************************************************************************/
+
+    FMultichannelLoudnessAnalyzer::FMultichannelLoudnessAnalyzer(float InSampleRate, const FLoudnessAnalyzerSettings& InSettings)
+    :   Analyzer(InSampleRate, InSettings)
+    {
+        // Initialize all channel weights to default to 1.0f
+        ChannelWeights.Init(1.0f, EAudioMixerChannel::ChannelTypeCount);
+
+        // If elevation angle is less than 30 degrees and azimuth is between 60 and 120, set channel weight to 1.5f. 
+        ChannelWeights[EAudioMixerChannel::BackLeft] = 1.5f;
+        ChannelWeights[EAudioMixerChannel::BackRight] = 1.5f;
+
+        MonoBuffer.Reset(InSettings.FFTSize);
+        MonoBuffer.AddUninitialized(InSettings.FFTSize);
+    }
+
+    float FMultichannelLoudnessAnalyzer::CalculatePerceptualEnergy(TArrayView<const float> InView, const int32 InNumChannels, TArray<float>& OutChannelEnergies)
+    {
+		check(InView.Num() == (InNumChannels * Analyzer.GetSettings().FFTSize));
+
+        OutChannelEnergies.Reset(InNumChannels);
+        OutChannelEnergies.AddUninitialized(InNumChannels);
+
+        // Deinterleave audio and send to mono loudness analyzer.
+		TAutoDeinterleaveView<float, FAudioBufferAlignedAllocator> DeinterleaveView(InView, MonoBuffer, InNumChannels);
+		for (TAutoDeinterleaveView<float>::TChannel<FAudioBufferAlignedAllocator> Channel : DeinterleaveView)
+		{
+            OutChannelEnergies[Channel.ChannelIndex] = Analyzer.CalculatePerceptualEnergy(Channel.Values);
+		}
+				
+        // Combine channel energies into overall perceptual energy
+        float Energy = 0.0f;
+        for (int32 ChannelNum = 0; ChannelNum < InNumChannels; ChannelNum++)
+        {
+            Energy += ChannelWeights[ChannelNum] * OutChannelEnergies[ChannelNum];
+        }
+
+        return Energy;
+    }
+
+    float FMultichannelLoudnessAnalyzer::CalculateLoudness(TArrayView<const float> InView, const int32 InNumChannels, TArray<float>& OutChannelLoudness)
+    {
+        // Convert perceptual energy to loudness 
+        float Energy = CalculatePerceptualEnergy(InView, InNumChannels, OutChannelLoudness);
+        for (float& Value : OutChannelLoudness)
+        {
+            Value = FLoudnessAnalyzer::ConvertPerceptualEnergyToLoudness(Value);
+        }
+        return FLoudnessAnalyzer::ConvertPerceptualEnergyToLoudness(Energy);
+    }
+
+    const FLoudnessAnalyzerSettings& FMultichannelLoudnessAnalyzer::GetSettings() const
+    {
+        return Analyzer.GetSettings();
+    }
+}
+
diff --git a/Engine/Plugins/Runtime/AudioSynesthesia/Source/AudioSynesthesiaCore/Private/LoudnessNRTFactory.cpp b/Engine/Plugins/Runtime/AudioSynesthesia/Source/AudioSynesthesiaCore/Private/LoudnessNRTFactory.cpp
new file mode 100644
--- /dev/null
+++ b/Engine/Plugins/Runtime/AudioSynesthesia/Source/AudioSynesthesiaCore/Private/LoudnessNRTFactory.cpp
@@ -0,0 +1,275 @@
+// Copyright Epic Games, Inc. All Rights Reserved.
+
+#include "LoudnessNRTFactory.h"
+#include "AudioSynesthesiaCustomVersion.h"
+#include "DSP/SlidingWindow.h"
+
+namespace Audio
+{
+	/************************************************************************/
+	/********************** FLoudnessNRTResult ******************************/
+	/************************************************************************/
+	FArchive &operator <<(FArchive& Ar, FLoudnessDatum& Datum)
+	{
+		Ar.UsingCustomVersion(FAudioSynesthesiaCustomVersion::GUID);
+
+		Ar << Datum.Channel;
+		Ar << Datum.Timestamp;
+		Ar << Datum.Energy;
+		Ar << Datum.Loudness;
+
+		return Ar;
+	}
+
+	// Use for overall loudness storage
+	const int32 FLoudnessNRTResult::ChannelIndexOverall = -1;
+
+
+	FLoudnessNRTResult::FLoudnessNRTResult()
+	:	DurationInSeconds(0.f)
+	,	bIsSortedChronologically(false)
+	{
+	}
+
+
+	void FLoudnessNRTResult::Serialize(FArchive& Archive)
+	{
+		Archive.UsingCustomVersion(FAudioSynesthesiaCustomVersion::GUID);
+
+		Archive << DurationInSeconds;
+		Archive << bIsSortedChronologically;
+		Archive << ChannelLoudnessArrays;
+		Archive << ChannelLoudnessIntervals;
+	}
+
+	void FLoudnessNRTResult::Add(const FLoudnessDatum& InDatum)
+	{
+		// Store loudness data in appropriate channel
+		TArray<FLoudnessDatum>& LoudnessArray = ChannelLoudnessArrays.FindOrAdd(InDatum.Channel);
+		LoudnessArray.Add(InDatum);
+
+		// Update intervals per channel
+		FFloatInterval& Interval = ChannelLoudnessIntervals.FindOrAdd(InDatum.Channel);
+		Interval.Include(InDatum.Loudness);
+
+		// Mark as not sorted
+		bIsSortedChronologically = false;
+	}
+
+	bool FLoudnessNRTResult::ContainsChannel(int32 InChannelIndex) const
+	{
+		return ChannelLoudnessArrays.Contains(InChannelIndex);
+	}
+
+	const TArray<FLoudnessDatum>& FLoudnessNRTResult::GetChannelLoudnessArray(int32 ChannelIdx) const
+	{
+		return ChannelLoudnessArrays[ChannelIdx];
+	}
+
+	const TArray<FLoudnessDatum>& FLoudnessNRTResult::GetLoudnessArray() const
+	{
+		return ChannelLoudnessArrays[ChannelIndexOverall];
+	}
+
+	float FLoudnessNRTResult::GetLoudnessRange(float InNoiseFloor) const
+	{
+		return GetChannelLoudnessRange(ChannelIndexOverall, InNoiseFloor);
+	}
+
+	float FLoudnessNRTResult::GetChannelLoudnessRange(int32 InChannelIdx, float InNoiseFloor) const
+	{
+		const FFloatInterval& Interval = ChannelLoudnessIntervals.FindRef(InChannelIdx);
+
+		if (Interval.Contains(InNoiseFloor))
+		{
+			// Noise floor is greater than minimum, so range must be relative to noise floor
+			return Interval.Max - InNoiseFloor;
+		}
+		else if (InNoiseFloor > Interval.Max)
+		{
+			return 0.f;
+		}
+
+		// Noise floor is less than minimum, so use entire range.
+		return Interval.Size();
+	}
+
+	void FLoudnessNRTResult::GetChannels(TArray<int32>& OutChannels) const
+	{
+		ChannelLoudnessArrays.GetKeys(OutChannels);
+	}
+
+	float FLoudnessNRTResult::GetDurationInSeconds() const 
+	{
+		return DurationInSeconds;
+	}
+
+	void FLoudnessNRTResult::SetDurationInSeconds(float InDuration)
+	{
+		DurationInSeconds = InDuration;
+	}
+
+	bool FLoudnessNRTResult::IsSortedChronologically() const
+	{
+		return bIsSortedChronologically;
+	}
+
+	void FLoudnessNRTResult::SortChronologically()
+	{
+		for (auto& KeyValue : ChannelLoudnessArrays)
+		{
+			KeyValue.Value.Sort([](const FLoudnessDatum& A, const FLoudnessDatum& B) { return A.Timestamp < B.Timestamp; });
+		}
+
+		bIsSortedChronologically = true;
+	}
+
+	/************************************************************************/
+	/********************** FLoudnessNRTWorker ******************************/
+	/************************************************************************/
+
+	FLoudnessNRTWorker::FLoudnessNRTWorker(const FAnalyzerNRTParameters& InParams, const FLoudnessNRTSettings& InAnalyzerSettings)
+	: NumChannels(InParams.NumChannels)
+	, NumAnalyzedBuffers(0)
+	, NumHopFrames(0)
+	, NumFrames(0)
+	, SampleRate(InParams.SampleRate)
+	{
+		check(NumChannels > 0);
+		check(SampleRate > 0.f);
+
+		// Include NumChannels in calculating window size because windows are generated 
+		// with interleaved samples and deinterleaved later. 
+		int32 NumWindowSamples = InAnalyzerSettings.FFTSize * NumChannels;
+
+		NumHopFrames = FMath::CeilToInt(InAnalyzerSettings.AnalysisPeriod * InParams.SampleRate);
+
+		NumHopFrames = FMath::Max(1, NumHopFrames);
+
+		int32 NumHopSamples = NumHopFrames * NumChannels;
+
+		InternalBuffer = MakeUnique<TSlidingBuffer<float>>(NumWindowSamples, NumHopSamples);
+
+		Analyzer = MakeUnique<FMultichannelLoudnessAnalyzer>(InParams.SampleRate, InAnalyzerSettings);
+	}
+
+	void FLoudnessNRTWorker::Analyze(TArrayView<const float> InAudio, IAnalyzerNRTResult* OutResult) 
+	{
+		NumFrames += InAudio.Num() / NumChannels;
+
+		// Assume that outer layers ensured that this is of correct type.
+		FLoudnessNRTResult* LoudnessResult = static_cast<FLoudnessNRTResult*>(OutResult);
+
+		check(nullptr != LoudnessResult);
+		
+		TAutoSlidingWindow<float> SlidingWindow(*InternalBuffer, InAudio, InternalWindow);
+
+		// Loop through entire array of input samples. 
+		for (const TArray<float>& Window : SlidingWindow)
+		{
+			AnalyzeWindow(Window, *LoudnessResult);
+		}
+	}
+
+	// Called when analysis of audio asset is complete. 
+	void FLoudnessNRTWorker::Finalize(IAnalyzerNRTResult* OutResult) 
+	{
+		FLoudnessNRTResult* LoudnessResult = static_cast<FLoudnessNRTResult*>(OutResult);
+
+		check(nullptr != LoudnessResult);
+
+		TArray<float> EmptyArray;
+		bool bDoFlush = true; // Add zeros to last window until all data has been processed.
+		TAutoSlidingWindow<float> SlidingWindow(*InternalBuffer, EmptyArray, InternalWindow, bDoFlush);
+
+		// Loop through entire array of input samples. 
+		for (const TArray<float>& Window : SlidingWindow)
+		{
+			AnalyzeWindow(Window, *LoudnessResult);
+		}
+
+		// Prepare result for fast lookup.
+		LoudnessResult->SortChronologically(); 
+
+		LoudnessResult->SetDurationInSeconds(static_cast<float>(NumFrames) / SampleRate);
+
+		// Reset internal counters
+		NumAnalyzedBuffers = 0; 
+		NumFrames = 0;
+
+		// Reset sliding buffer
+		InternalBuffer->Reset();
+	}
+
+	void FLoudnessNRTWorker::AnalyzeWindow(TArrayView<const float> InWindow, FLoudnessNRTResult& OutResult)
+	{
+		TArray<float> ChannelPerceptualEnergy;
+		float OverallPerceptualEnergy = 0.0f;
+
+		// Calculate perceptual energy
+		OverallPerceptualEnergy = Analyzer->CalculatePerceptualEnergy(InWindow, NumChannels, ChannelPerceptualEnergy);
+
+		// Convert to perceptual energy to loudness
+		TArray<float> ChannelLoudness;
+		float OverallLoudness = 0.f;
+		
+		OverallLoudness = FLoudnessAnalyzer::ConvertPerceptualEnergyToLoudness(OverallPerceptualEnergy);
+
+		for (float PerceptualEnergy : ChannelPerceptualEnergy)
+		{
+			ChannelLoudness.Add(FLoudnessAnalyzer::ConvertPerceptualEnergyToLoudness(PerceptualEnergy));
+		}
+
+		// Calculate timestamp
+		float Timestamp = ((NumAnalyzedBuffers * NumHopFrames) + (0.5f * Analyzer->GetSettings().FFTSize)) / SampleRate;
+
+		// Add overall FLoudnessDatum to result
+		const FLoudnessDatum OverallDatum = {FLoudnessNRTResult::ChannelIndexOverall, Timestamp, OverallPerceptualEnergy, OverallLoudness};
+		OutResult.Add(OverallDatum);
+
+		// Add channel FLoudnessDatum to result
+		for (int32 ChannelIdx = 0; ChannelIdx < NumChannels; ChannelIdx++)
+		{
+			const FLoudnessDatum ChannelDatum = {ChannelIdx, Timestamp, ChannelPerceptualEnergy[ChannelIdx], ChannelLoudness[ChannelIdx]};
+
+			OutResult.Add(ChannelDatum);
+		}	
+
+		// update counters
+		NumAnalyzedBuffers++;
+	}
+
+
+	/************************************************************************/
+	/************************* FLoudnessNRTFactory **************************/
+	/************************************************************************/
+
+	/** Name of specific analyzer type. */
+	FName FLoudnessNRTFactory::GetName() const 
+	{
+		static FName FactoryName(TEXT("LoudnessNRTFactory"));
+		return FactoryName;
+	}
+
+	/** Human readable name of analyzer. */
+	FString FLoudnessNRTFactory::GetTitle() const
+	{
+		return TEXT("Loudness Analyzer Non-Real-Time");
+	}
+
+	TUniquePtr<IAnalyzerNRTResult> FLoudnessNRTFactory::NewResult() const
+	{
+		TUniquePtr<FLoudnessNRTResult> Result = MakeUnique<FLoudnessNRTResult>();
+		return Result;
+	}
+
+	TUniquePtr<IAnalyzerNRTWorker> FLoudnessNRTFactory::NewWorker(const FAnalyzerNRTParameters& InParams, const IAnalyzerNRTSettings* InSettings) const
+	{
+		const FLoudnessNRTSettings* LoudnessSettings = static_cast<const FLoudnessNRTSettings*>(InSettings);
+
+		check(nullptr != LoudnessSettings);
+
+		return MakeUnique<FLoudnessNRTWorker>(InParams, *LoudnessSettings);
+	}
+}
+
diff --git a/Engine/Plugins/Runtime/AudioSynesthesia/Source/AudioSynesthesiaCore/Private/OnsetAnalyzer.cpp b/Engine/Plugins/Runtime/AudioSynesthesia/Source/AudioSynesthesiaCore/Private/OnsetAnalyzer.cpp
new file mode 100644
--- /dev/null
+++ b/Engine/Plugins/Runtime/AudioSynesthesia/Source/AudioSynesthesiaCore/Private/OnsetAnalyzer.cpp
@@ -0,0 +1,236 @@
+// Copyright Epic Games, Inc. All Rights Reserved.
+
+#include "OnsetAnalyzer.h"
+#include "PeakPicker.h"
+#include "DSP/FloatArrayMath.h"
+
+namespace Audio
+{
+	FOnsetStrengthAnalyzer::FOnsetStrengthAnalyzer(const FOnsetStrengthSettings& InSettings, float InSampleRate)
+	:	Settings(InSettings)
+	,	LagSpectraIndex(0)
+	,	SlidingBuffer(InSettings.NumWindowFrames, InSettings.NumHopFrames)
+	,	Window(Settings.WindowType, Settings.NumWindowFrames, 1, false)
+	{
+		// Check values legit.
+		check(InSampleRate > 0.f);
+		check(Settings.FFTSize > 0);
+		check(Settings.NumHopFrames > 0);
+		check(Settings.NumWindowFrames > 0);
+		check(Settings.ComparisonLag > 0);
+		check(Settings.MelSettings.NumBands > 0);
+		check(Settings.NumWindowFrames <= Settings.FFTSize);
+
+		// intialize buffers dealing with audio window
+		WorkingBuffer.AddUninitialized(Settings.NumWindowFrames);
+
+		// initialize buffers dealing with fft input / output
+		WindowedSamples.AddZeroed(Settings.FFTSize);
+		FFTOutputRealData.AddUninitialized(Settings.FFTSize);
+		FFTOutputImagData.AddUninitialized(Settings.FFTSize);
+		FFTSpectrumBuffer.AddUninitialized(Settings.FFTSize);
+		
+		// intialize buffers dealing with perceptual spectrum 
+		PreviousSpectra.AddDefaulted(Settings.ComparisonLag);
+		MelSpectrum.AddUninitialized(Settings.MelSettings.NumBands);
+		SpectrumDifference.AddUninitialized(Settings.MelSettings.NumBands);
+
+		// create mel spectrum transform
+		MelTransform = NewMelSpectrumKernelTransform(Settings.MelSettings, Settings.FFTSize, InSampleRate);
+
+		check(MelTransform.IsValid());
+	}
+
+	void FOnsetStrengthAnalyzer::CalculateOnsetStrengths(TArrayView<const float> InSamples, TArray<float>& OutEnvelopeStrengths)
+	{
+		const bool bDoFlush = false;
+		AnalyzeAudio(InSamples, OutEnvelopeStrengths, bDoFlush);
+	}
+
+	void FOnsetStrengthAnalyzer::FlushAudio(TArray<float>& OutEnvelopeStrengths)
+	{
+		TArray<float> EmptyArray;
+
+		// Force flushing of internal audio buffer
+		const bool bDoFlush = true;
+
+		AnalyzeAudio(EmptyArray, OutEnvelopeStrengths, bDoFlush);
+	}
+	
+
+	void FOnsetStrengthAnalyzer::Reset()
+	{
+		// Reset internal memory
+		for (TArray<float>& Spectrum : PreviousSpectra)
+		{
+			Spectrum.Reset();
+		}
+		SlidingBuffer.Reset();
+	}
+
+	float FOnsetStrengthAnalyzer::GetTimestampForIndex(const FOnsetStrengthSettings& InSettings, float InSampleRate, int32 InIndex)
+	{
+		check(InSampleRate > 0.f);
+
+		// Getting the exact timestamp of an onset from the onset strength envelope is non-trivial because
+		// the onset strength envelope changes depending on the attack, decay and sustain of the onset as well
+		// as the choice of window, window size, etc. etc. 
+		//
+		// This approach does not attempt to account for all the subtle variability than can come up, but 
+		// uses a generally sane approach of estimating the timestamp to be somewhere between the center of
+		// the "loud" window and the end of the "quiet" window.
+		const int32 WindowOffset = InSettings.NumWindowFrames / 2;
+
+		const int32 LoudFrameIndex = (InIndex * InSettings.NumHopFrames) + WindowOffset;
+
+		// Get the end of the window to which this was compared.
+		const int32 QuietIndex = InIndex - InSettings.ComparisonLag;
+		const int32 QuietFrameIndex = (QuietIndex * InSettings.NumHopFrames) + InSettings.NumWindowFrames;
+
+		const int32 EstimatedOnsetFrameIndex = (LoudFrameIndex + QuietFrameIndex) / 2;
+
+		return static_cast<float>(EstimatedOnsetFrameIndex) / InSampleRate;
+	}
+
+	void FOnsetStrengthAnalyzer::AnalyzeAudio(TArrayView<const float> InSamples, TArray<float>& OutEnvelopeStrengths, bool bDoFlush)
+	{
+		typedef TAutoSlidingWindow<float, FAudioBufferAlignedAllocator> FSlidingWindow;
+
+		// Going to fill OutEnvelopeStrengths with envelope strength values generated during this call.
+		OutEnvelopeStrengths.Reset();
+
+		// Slide over audio 
+		FSlidingWindow SlidingWindow(SlidingBuffer, InSamples, WorkingBuffer, bDoFlush);
+		for (const AlignedFloatBuffer& WindowBuffer : SlidingWindow)
+		{
+			// Calculate onset strength per audio buffer.
+			float OnsetStrength = GetNextOnsetStrength(WindowBuffer);
+
+			OutEnvelopeStrengths.Add(OnsetStrength);
+		}
+	}
+
+	float FOnsetStrengthAnalyzer::GetNextOnsetStrength(TArrayView<const float> InSamples)
+	{
+		float OnsetStrength = 0.f;
+
+		// Copy input samples and apply window
+		FMemory::Memcpy(WindowedSamples.GetData(), InSamples.GetData(), sizeof(float) * Settings.NumWindowFrames);
+		Window.ApplyToBuffer(WindowedSamples.GetData());
+
+		// Take FFT of windowed samples
+		const FFTTimeDomainData TimeData = {WindowedSamples.GetData(), Settings.FFTSize};
+		FFTFreqDomainData FreqData = {FFTOutputRealData.GetData(), FFTOutputImagData.GetData()};
+
+		PerformFFT(TimeData, FreqData);
+
+		// Calculate spectrum from FFT output
+		ComputePowerSpectrum(FreqData, Settings.FFTSize, FFTSpectrumBuffer);
+		
+		// Convert spectrum to CQT
+		MelTransform->TransformArray(FFTSpectrumBuffer, MelSpectrum);
+
+		// Apply decibel scaling 
+		ArrayPowerToDecibelInPlace(MelSpectrum, -90.f);
+
+		// Clamp to noise floor
+		ArrayClampMinInPlace(MelSpectrum, Settings.NoiseFloorDb);
+
+		// Onset strength is mean(abs(diff(spectrum - previous_spectrum)))
+		TArray<float>& PreviousSpectrum = PreviousSpectra[LagSpectraIndex];
+		if (PreviousSpectrum.Num() > 0)
+		{
+			ArraySubtract(MelSpectrum, PreviousSpectrum, SpectrumDifference);
+
+			// Half wave rectify
+			ArrayClampMinInPlace(SpectrumDifference, 0.f);
+
+			ArrayMean(SpectrumDifference, OnsetStrength);
+		}
+
+		// Save this spectrum for later use
+		PreviousSpectrum.Reset(Settings.MelSettings.NumBands);
+		PreviousSpectrum.AddUninitialized(Settings.MelSettings.NumBands);
+		FMemory::Memcpy(PreviousSpectrum.GetData(), MelSpectrum.GetData(), sizeof(float) * Settings.MelSettings.NumBands);
+
+		// Increment index for comparison spectra
+		LagSpectraIndex = (LagSpectraIndex + 1) % Settings.ComparisonLag;
+
+		return OnsetStrength;
+	}
+
+
+	/*******************************************************************************/
+	/*******************************************************************************/
+	/*******************************************************************************/
+
+	void OnsetExtractIndices(const FPeakPickerSettings& InSettings, TArrayView<const float> InOnsetEnvelope, TArray<int32>& OutOnsetIndices)
+	{
+		const int32 Num = InOnsetEnvelope.Num();
+		OutOnsetIndices.Reset();
+
+		if (Num < 1)
+		{
+			return;
+		}
+		
+		TArray<float> NormOnsets;
+
+		// Onset strength is normalized between [0, 1.0] before picking peaks.
+		ArrayMinMaxNormalize(InOnsetEnvelope, NormOnsets);
+
+		FPeakPicker PeakPicker(InSettings);
+		PeakPicker.PickPeaks(NormOnsets, OutOnsetIndices);
+	}
+
+	/*******************************************************************************/
+	/*******************************************************************************/
+	/*******************************************************************************/
+
+	void OnsetBacktrackIndices(TArrayView<const float> InOnsetStrength, TArrayView<const int32> InOnsetIndices, TArray<int32>& OutBacktrackedOnsetIndices)
+	{
+		const int32 NumOnsets = InOnsetIndices.Num();
+		const int32 NumOnsetStrength = InOnsetStrength.Num();
+
+		OutBacktrackedOnsetIndices.Reset();
+		
+		if (NumOnsets < 1)
+		{
+			return;
+		}
+
+		OutBacktrackedOnsetIndices.AddUninitialized(NumOnsets);
+
+		const int32* InOnsetIndicesData = InOnsetIndices.GetData();
+		int32* OutOnsetIndicesData = OutBacktrackedOnsetIndices.GetData();
+		const float* InOnsetStrengthData = InOnsetStrength.GetData();
+
+		for (int32 i = 0; i < NumOnsets; i++)
+		{
+
+			int32 OnsetStartIndex = InOnsetIndicesData[i];
+			check(OnsetStartIndex < NumOnsetStrength);
+
+			while (OnsetStartIndex > 0)
+			{
+				if (InOnsetStrengthData[OnsetStartIndex] <= 0.f)
+				{
+					// No onset strength means the audio is steady state or silent.
+					break;
+				}
+				else if (InOnsetStrengthData[OnsetStartIndex - 1] <= InOnsetStrengthData[OnsetStartIndex])
+				{
+					// The onest strength is still holding steady or decreasing.
+					OnsetStartIndex--;
+				}
+				else
+				{
+					// Onset strength is no longer decreasing.
+					break;
+				}
+			}
+
+			OutOnsetIndicesData[i] = OnsetStartIndex;
+		}
+	}
+}
diff --git a/Engine/Plugins/Runtime/AudioSynesthesia/Source/AudioSynesthesiaCore/Private/OnsetNRTFactory.cpp b/Engine/Plugins/Runtime/AudioSynesthesia/Source/AudioSynesthesiaCore/Private/OnsetNRTFactory.cpp
new file mode 100644
--- /dev/null
+++ b/Engine/Plugins/Runtime/AudioSynesthesia/Source/AudioSynesthesiaCore/Private/OnsetNRTFactory.cpp
@@ -0,0 +1,326 @@
+// Copyright Epic Games, Inc. All Rights Reserved.
+
+#include "OnsetNRTFactory.h"
+#include "OnsetAnalyzer.h"
+#include "AudioSynesthesiaCustomVersion.h"
+#include "DSP/DeinterleaveView.h"
+#include "DSP/FloatArrayMath.h"
+
+namespace Audio 
+{
+	/*******************************************************************************/
+	/*************************** FOnsetNRTResult *******************************/
+	/*******************************************************************************/
+
+	FArchive &operator <<(FArchive& Ar, FOnset& Frame)
+	{
+		Ar.UsingCustomVersion(FAudioSynesthesiaCustomVersion::GUID);
+
+		Ar << Frame.Channel;
+		Ar << Frame.Timestamp;
+		Ar << Frame.Strength;
+
+		return Ar;
+	}
+
+	FOnsetNRTResult::FOnsetNRTResult()
+	:	DurationInSeconds(0.f)
+	,	bIsSortedChronologically(false)
+	{}
+
+	void FOnsetNRTResult::Serialize(FArchive& Archive)
+	{
+		Archive.UsingCustomVersion(FAudioSynesthesiaCustomVersion::GUID);
+
+		Archive << DurationInSeconds;
+		Archive << bIsSortedChronologically;
+		Archive << ChannelOnsets;
+		Archive << ChannelOnsetStrengthIntervals;
+	}
+
+	void FOnsetNRTResult::AddOnset(int32 InChannelIndex, float InTimestamp, float InStrength)
+	{
+		// Add to onset arrays
+		TArray<FOnset>& OnsetArray = ChannelOnsets.FindOrAdd(InChannelIndex);
+		OnsetArray.Emplace(InChannelIndex, InTimestamp, InStrength);
+
+		// Store min/max value
+		FFloatInterval& Interval = ChannelOnsetStrengthIntervals.FindOrAdd(InChannelIndex);
+		Interval.Include(InStrength);
+
+		bIsSortedChronologically = false;
+	}
+
+	void FOnsetNRTResult::AddChannel(int32 InChannelIndex)
+	{
+		// Add an empty channel. Useful when there is a channel with no onsets.
+		ChannelOnsets.FindOrAdd(InChannelIndex);
+		ChannelOnsetStrengthIntervals.FindOrAdd(InChannelIndex);
+	}
+
+	bool FOnsetNRTResult::ContainsChannel(int32 InChannelIndex) const
+	{
+		return ChannelOnsets.Contains(InChannelIndex);
+	}
+
+	const TArray<FOnset>& FOnsetNRTResult::GetOnsetsForChannel(int32 InChannelIndex) const
+	{
+		check(ChannelOnsets.Contains(InChannelIndex));
+
+		return ChannelOnsets[InChannelIndex];
+	}
+
+	FFloatInterval FOnsetNRTResult::GetChannelOnsetInterval(int32 InChannelIndex) const
+	{
+		check(ChannelOnsets.Contains(InChannelIndex));
+
+		return ChannelOnsetStrengthIntervals[InChannelIndex];
+	}
+
+	void FOnsetNRTResult::GetChannels(TArray<int32>& OutChannels) const
+	{
+		ChannelOnsets.GetKeys(OutChannels);
+	}
+
+	float FOnsetNRTResult::GetDurationInSeconds() const 
+	{
+		return DurationInSeconds;
+	}
+
+	void FOnsetNRTResult::SetDurationInSeconds(float InDuration)
+	{
+		DurationInSeconds = InDuration;
+	}
+
+	 /* Returns true if FOnsetNRTFrame arrays are sorted in chronologically ascending order via their timestamp.
+	 */
+	bool FOnsetNRTResult::IsSortedChronologically() const
+	{
+		return bIsSortedChronologically;
+	}
+
+	/**
+	 * Sorts FOnsetNRTFrame arrays in chronologically ascnding order via their timestamp.
+	 */
+	void FOnsetNRTResult::SortChronologically()
+	{
+		for (auto& KeyValue : ChannelOnsets)
+		{
+			KeyValue.Value.Sort([](const FOnset& A, const FOnset& B) { return A.Timestamp < B.Timestamp; });
+		}
+
+		bIsSortedChronologically = true;
+	}
+
+	/*******************************************************************************/
+	/*************************** FOnsetNRTWorker *******************************/
+	/*******************************************************************************/
+
+	FOnsetNRTWorker::FOnsetNRTWorker(const FAnalyzerNRTParameters& InParams, const FOnsetNRTSettings& InAnalyzerSettings)
+	:	Settings(InAnalyzerSettings)
+	,	NumFrames(0)
+	,	NumChannels(InParams.NumChannels)
+	,	SampleRate(InParams.SampleRate)
+	,	NumWindowFrames(4096)
+	,	NumWindowSamples(0)
+	,	MonoScaling(1.f)
+	{
+		check(NumChannels > 0);
+		check(SampleRate > 0.f);
+
+		if (Settings.bDownmixToMono)
+		{
+			MonoBuffer.Reset(NumWindowFrames);
+			MonoBuffer.AddUninitialized(NumWindowFrames);
+			// equal power mixing
+			MonoScaling = 1.f / FMath::Sqrt(static_cast<float>(NumChannels));
+
+			// Only need one onset strength analyzer
+			TUniquePtr<FOnsetStrengthAnalyzer> Analyzer = MakeUnique<FOnsetStrengthAnalyzer>(InAnalyzerSettings.OnsetStrengthSettings, SampleRate);
+			OnsetStrengthAnalyzers.Emplace(MoveTemp(Analyzer));
+
+			// Only need one onset strength array
+			OnsetStrengths.AddDefaulted(1);
+		}
+		else
+		{
+			// Make individual onset strength analyzer for each channel
+			for (int32 i = 0; i < NumChannels; i++)
+			{
+				TUniquePtr<FOnsetStrengthAnalyzer> Analyzer = MakeUnique<FOnsetStrengthAnalyzer>(InAnalyzerSettings.OnsetStrengthSettings, SampleRate);
+				OnsetStrengthAnalyzers.Emplace(MoveTemp(Analyzer));
+			}
+			OnsetStrengths.AddDefaulted(NumChannels);
+		}
+
+		NumWindowSamples = NumChannels * NumWindowFrames;
+		
+		// Sliding buffer used for deinterleaving audio
+		SlidingBuffer = MakeUnique<TSlidingBuffer<float> >(NumWindowSamples, NumWindowSamples);
+	}
+
+
+	void FOnsetNRTWorker::Analyze(TArrayView<const float> InAudio, IAnalyzerNRTResult* OutResult)
+	{
+		bool bDoFlush = false;
+
+		AnalyzeMultichannel(InAudio, OutResult, bDoFlush);
+	}
+
+	void FOnsetNRTWorker::AnalyzeMultichannel(TArrayView<const float> InAudio, IAnalyzerNRTResult* OutResult, bool bDoFlush)
+	{
+		NumFrames += InAudio.Num() / NumChannels;
+
+		// Assume that outer layers ensured that this is of correct type.
+		FOnsetNRTResult* OnsetResult = static_cast<FOnsetNRTResult*>(OutResult);
+
+		check(nullptr != OnsetResult);
+
+		TAutoSlidingWindow<float, FAudioBufferAlignedAllocator> SlidingWindow(*SlidingBuffer, InAudio, HopBuffer, bDoFlush);
+
+		if (Settings.bDownmixToMono)
+		{
+			for (const AlignedFloatBuffer& Window : SlidingWindow)
+			{
+				FMemory::Memset(MonoBuffer.GetData(), 0, sizeof(float) * NumWindowFrames);
+
+				// Sum all channels into mono buffer
+				TAutoDeinterleaveView<float, FAudioBufferAlignedAllocator> DeinterleaveView(Window, ChannelBuffer, NumChannels);
+				for(auto Channel : DeinterleaveView)
+				{
+					MixInBufferFast(Channel.Values, MonoBuffer);		
+				}
+
+				// Equal power sum. assuming incoherrent signals.
+				MultiplyBufferByConstantInPlace(MonoBuffer, 1.f / FMath::Sqrt(static_cast<float>(NumChannels)));
+
+				AnalyzeWindow(MonoBuffer, 0, *OnsetResult);
+			}
+		}
+		else
+		{
+			// Loop through entire array of input samples.
+			for (const AlignedFloatBuffer& Window : SlidingWindow)
+			{
+				// Each channel is analyzed seperately.
+				for(const TDeinterleaveView<float>::TChannel<FAudioBufferAlignedAllocator>& Channel : TAutoDeinterleaveView<float, FAudioBufferAlignedAllocator>(Window, ChannelBuffer, NumChannels))
+				{
+					AnalyzeWindow(Channel.Values, Channel.ChannelIndex, *OnsetResult);
+				}
+			}
+		}
+	}
+
+	/** Called when analysis of audio asset is complete. */
+	void FOnsetNRTWorker::Finalize(IAnalyzerNRTResult* OutResult) 
+	{
+		// Finish analyzing audio
+		AlignedFloatBuffer EmptyArray;
+		bool bDoFlush = true;
+
+		AnalyzeMultichannel(EmptyArray, OutResult, bDoFlush);
+
+		// Call flush on all the onset strength analyzers
+		for (int32 i = 0; i < OnsetStrengthAnalyzers.Num(); i++)
+		{
+			TArray<float> Strengths;
+			OnsetStrengthAnalyzers[i]->FlushAudio(Strengths);
+			OnsetStrengths[i].Append(Strengths);
+		}
+
+		// Assume that outer layers ensured that this is of correct type.
+		FOnsetNRTResult* OnsetResult = static_cast<FOnsetNRTResult*>(OutResult);
+
+		check(nullptr != OnsetResult);
+
+		if (nullptr != OnsetResult)
+		{
+			ExtractOnsetsFromOnsetStrengths(*OnsetResult);
+
+			float DurationInSeconds = static_cast<float>(NumFrames) / SampleRate;
+
+			OnsetResult->SetDurationInSeconds(DurationInSeconds);
+
+			OnsetResult->SortChronologically();
+		}
+
+		// Reset internal objects
+		Reset();
+	}
+
+	void FOnsetNRTWorker::AnalyzeWindow(const AlignedFloatBuffer& InWindow, int32 InChannelIndex, FOnsetNRTResult& OutResult)
+	{
+		TArray<float> Strengths;
+		OnsetStrengthAnalyzers[InChannelIndex]->CalculateOnsetStrengths(InWindow, Strengths);
+
+		OnsetStrengths[InChannelIndex].Append(Strengths);
+	}
+
+	void FOnsetNRTWorker::ExtractOnsetsFromOnsetStrengths(FOnsetNRTResult& OutResult)
+	{
+		for (int32 ChannelIndex = 0; ChannelIndex < OnsetStrengths.Num(); ChannelIndex++)
+		{
+			TArray<int32> OnsetIndices;
+			TArray<float>& ChannelStrengths = OnsetStrengths[ChannelIndex];
+
+			OnsetExtractIndices(Settings.PeakPickerSettings, ChannelStrengths, OnsetIndices);
+
+			// Call this in case there are no onsets. Then the results object will know about the
+			// channel, and will know there are no onsets. Otherwise it would be ambiguous whether 
+			// the channel doesn't exist or if it does exist but contains no onsets.
+			OutResult.AddChannel(ChannelIndex);
+
+			for (int32 OnsetIndex : OnsetIndices)
+			{
+				float Timestamp = FOnsetStrengthAnalyzer::GetTimestampForIndex(Settings.OnsetStrengthSettings, SampleRate, OnsetIndex);
+				OutResult.AddOnset(ChannelIndex, Timestamp, OnsetStrengths[ChannelIndex][OnsetIndex]);
+			}
+		}
+	}
+
+	void FOnsetNRTWorker::Reset()
+	{
+		// Reset internal counters
+		NumFrames = 0;
+		SlidingBuffer.Reset();
+
+		for (int32 i = 0; i < OnsetStrengthAnalyzers.Num(); i++)
+		{
+			OnsetStrengthAnalyzers[i]->Reset();
+		}
+		for (int32 i = 0; i < OnsetStrengths.Num(); i++)
+		{
+			OnsetStrengths[i].Reset();
+		}
+	}
+
+	// Name of this analyzer type.
+	FName FOnsetNRTFactory::GetName() const
+	{
+		static FName FactoryName(TEXT("OnsetNRTFactory"));
+		return FactoryName;
+	}
+
+	// Human readable name of this analyzer.
+	FString FOnsetNRTFactory::GetTitle() const
+	{
+		return TEXT("Constant Q Analyzer Non-Real-Time");
+	}
+
+	// Create a new FOnsetNRTResult.
+	TUniquePtr<IAnalyzerNRTResult> FOnsetNRTFactory::NewResult() const
+	{
+		return MakeUnique<FOnsetNRTResult>();
+	}
+
+	// Create a new FOnsetNRTWorker
+	TUniquePtr<IAnalyzerNRTWorker> FOnsetNRTFactory::NewWorker(const FAnalyzerNRTParameters& InParams, const IAnalyzerNRTSettings* InSettings) const
+	{
+		const FOnsetNRTSettings* OnsetSettings = static_cast<const FOnsetNRTSettings*>(InSettings);
+
+		check(nullptr != OnsetSettings);
+
+		return MakeUnique<FOnsetNRTWorker>(InParams, *OnsetSettings);
+	}
+}
+
diff --git a/Engine/Plugins/Runtime/AudioSynesthesia/Source/AudioSynesthesiaCore/Private/PeakPicker.cpp b/Engine/Plugins/Runtime/AudioSynesthesia/Source/AudioSynesthesiaCore/Private/PeakPicker.cpp
new file mode 100644
--- /dev/null
+++ b/Engine/Plugins/Runtime/AudioSynesthesia/Source/AudioSynesthesiaCore/Private/PeakPicker.cpp
@@ -0,0 +1,66 @@
+// Copyright Epic Games, Inc. All Rights Reserved.
+
+#include "PeakPicker.h"
+#include "CoreMinimal.h"
+#include "DSP/FloatArrayMath.h"
+
+namespace Audio
+{
+	FPeakPicker::FPeakPicker(const FPeakPickerSettings& InSettings)
+	:	Settings(InSettings)
+	{
+		// Clamp settings to valid values
+		Settings.NumPreMax = FMath::Max(1, Settings.NumPreMax);
+		Settings.NumPostMax = FMath::Max(1, Settings.NumPostMax);
+		Settings.NumPreMean = FMath::Max(1, Settings.NumPreMean);
+		Settings.NumPostMean = FMath::Max(1, Settings.NumPostMean);
+		Settings.NumWait = FMath::Max(1, Settings.NumWait);
+		Settings.MeanDelta = FMath::Clamp(Settings.MeanDelta, 1e-6f, INFINITY);
+	}
+
+	void FPeakPicker::PickPeaks(TArrayView<const float> InData, TArray<int32>& OutPeakIndices)
+	{
+		OutPeakIndices.Empty();
+		if (InData.Num() < 1)
+		{
+			return;
+		}
+
+		// Get max filtered version
+		TArray<float> MaxFilteredData;
+		int32 MaxFilterWindowSize = Settings.NumPreMax + Settings.NumPostMax + 1;
+		int32 MaxFilterWindowOrigin = Settings.NumPreMax;
+
+		ArrayMaxFilter(InData, MaxFilterWindowSize, MaxFilterWindowOrigin, MaxFilteredData);
+
+		
+		// Get mean filtered version
+		TArray<float> MeanFilteredData;
+		int32 MeanFilterWindowSize = Settings.NumPreMean + Settings.NumPostMean + 1;
+		int32 MeanFilterWindowOrigin = Settings.NumPreMean;
+
+		ArrayMeanFilter(InData, MeanFilterWindowSize, MeanFilterWindowOrigin, MeanFilteredData);
+
+		// Find peak locations
+		const float* InDataPtr = InData.GetData();
+		const float* MaxDataPtr = MaxFilteredData.GetData();
+		const float* MeanDataPtr = MeanFilteredData.GetData();
+
+		const int32 Num = InData.Num();
+		int32 PossiblePeakIndex = 0;
+		while (PossiblePeakIndex < Num)
+		{
+			if (InDataPtr[PossiblePeakIndex] >= MaxDataPtr[PossiblePeakIndex])
+			{
+				if (InDataPtr[PossiblePeakIndex] >= (MeanDataPtr[PossiblePeakIndex] + Settings.MeanDelta))
+				{
+					OutPeakIndices.Add(PossiblePeakIndex);
+					PossiblePeakIndex += Settings.NumWait;
+					continue;
+				}
+			}
+			PossiblePeakIndex++;
+		}
+	}
+}
+
diff --git a/Engine/Plugins/Runtime/AudioSynesthesia/Source/AudioSynesthesiaCore/Public/AudioSynesthesiaCoreModule.h b/Engine/Plugins/Runtime/AudioSynesthesia/Source/AudioSynesthesiaCore/Public/AudioSynesthesiaCoreModule.h
new file mode 100644
--- /dev/null
+++ b/Engine/Plugins/Runtime/AudioSynesthesia/Source/AudioSynesthesiaCore/Public/AudioSynesthesiaCoreModule.h
@@ -0,0 +1,11 @@
+// Copyright Epic Games, Inc. All Rights Reserved.
+
+#include "CoreMinimal.h"
+#include "Modules/ModuleInterface.h"
+
+namespace Audio
+{
+	class AUDIOSYNESTHESIACORE_API IAudioSynesthesiaCoreModule : public IModuleInterface
+	{
+	};
+}
diff --git a/Engine/Plugins/Runtime/AudioSynesthesia/Source/AudioSynesthesiaCore/Public/ConstantQAnalyzer.h b/Engine/Plugins/Runtime/AudioSynesthesia/Source/AudioSynesthesiaCore/Public/ConstantQAnalyzer.h
new file mode 100644
--- /dev/null
+++ b/Engine/Plugins/Runtime/AudioSynesthesia/Source/AudioSynesthesiaCore/Public/ConstantQAnalyzer.h
@@ -0,0 +1,93 @@
+// Copyright Epic Games, Inc. All Rights Reserved.
+
+#pragma once
+
+#include "DSP/ConstantQ.h"
+#include "DSP/AudioFFT.h"
+
+namespace Audio
+{
+	/** Denotes which magnitude scaling to apply to output */
+	enum class EConstantQScaling : uint8
+	{
+		/** CQT spectrum values conform to ESpectrumType, either magnitude or power. */
+		Linear,
+
+		/** CQT spectrum values are in decibels. */
+		Decibel
+	};
+
+	/** FConstantQAnalyzerSettings
+	 *
+	 * 	Settings for the Constant Q Analyzer.
+	 */
+	struct AUDIOSYNESTHESIACORE_API FConstantQAnalyzerSettings : public FPseudoConstantQKernelSettings
+	{
+		/** Size of FFT describe in number of samples */
+		int32 FFTSize;
+
+		/** Type of window to be applied to input audio */
+        EWindowType WindowType;
+
+		/** Type of spectrum to use. */
+		ESpectrumType SpectrumType;
+
+		/** Type of scaling to use. */
+		EConstantQScaling Scaling;
+
+		FConstantQAnalyzerSettings()
+		:	FFTSize(1024)
+		,	WindowType(EWindowType::Blackman)
+		,	SpectrumType(ESpectrumType::PowerSpectrum)
+		,	Scaling(EConstantQScaling::Decibel)	
+		{}
+	};
+
+	/** FConstantQAnalyzer
+	 *
+	 *  The Constant Q Analyzer computs a pseudo constant q transform of the input audio.
+	 *  ConstantQ transforms produce frequency strength at logarithmically spaced frequency
+	 *  intervals (as opposed to the linearly spaced intervals of the FFT). Logarithmically
+	 *  spaced frequency intervals map nicely to musical scales.
+	 */
+	class AUDIOSYNESTHESIACORE_API FConstantQAnalyzer
+	{
+	public:
+
+		/** Construct a FConstantQAnalyzer
+		 *
+		 * 	InSettings contains the parameters used to configure the analyzer.
+		 * 	InSampleRate is the sample rate of the input audio.
+		 */
+		FConstantQAnalyzer(const FConstantQAnalyzerSettings& InSettings, const float InSampleRate);
+
+		/** Calcualte Constant Q Transform of input audio.
+		 *
+		 * InSamples is an array of input audio. It must be of length GetSettings().FFTSize.
+		 * OutCQT will be populated with the constants q transform values. It will have GetSettings().NumBands values.
+		 */
+		void CalculateCQT(const float* InSamples, TArray<float>& OutCQT);
+
+		/** Returns a reference to the settings used in this analyzer */
+		const FConstantQAnalyzerSettings& GetSettings() const;
+
+	private:
+
+		const FConstantQAnalyzerSettings Settings;
+
+		const float SampleRate;
+
+		/** Only the first (FFTSize / 2) + 1 FFT bins are used since FFT of real values is conjugate symmetric. */
+		const int32 NumUsefulFFTBins;
+
+		/** Scale factor to normalize between different FFT Sizes */
+		FWindow Window;
+
+		AlignedFloatBuffer WindowedSamples;
+		AlignedFloatBuffer FFTOutputRealData;
+		AlignedFloatBuffer FFTOutputImagData;
+		AlignedFloatBuffer SpectrumBuffer;
+
+		TUniquePtr<FContiguousSparse2DKernelTransform> CQTTransform;
+	};
+}
diff --git a/Engine/Plugins/Runtime/AudioSynesthesia/Source/AudioSynesthesiaCore/Public/ConstantQNRTFactory.h b/Engine/Plugins/Runtime/AudioSynesthesia/Source/AudioSynesthesiaCore/Public/ConstantQNRTFactory.h
new file mode 100644
--- /dev/null
+++ b/Engine/Plugins/Runtime/AudioSynesthesia/Source/AudioSynesthesiaCore/Public/ConstantQNRTFactory.h
@@ -0,0 +1,204 @@
+// Copyright Epic Games, Inc. All Rights Reserved.
+
+#pragma once
+
+#include "IAudioAnalyzerNRTInterface.h"
+#include "DSP/ConstantQ.h"
+#include "DSP/SlidingWindow.h"
+#include "DSP/FloatArrayMath.h"
+#include "ConstantQAnalyzer.h"
+
+namespace Audio
+{
+	/** FConstantQNRTSettings
+	 *
+	 *  Settings for the Constant Q Non-Real-Time Factory.
+	 */
+	class AUDIOSYNESTHESIACORE_API FConstantQNRTSettings : public IAnalyzerNRTSettings, public FConstantQAnalyzerSettings 
+	{
+	public:
+
+		/** Time, in seconds, between constant Q frames. */
+		float AnalysisPeriod;
+
+		/** 
+		 * If true, all channels are mixed together with equal power before analysis. Only one channel
+		 * is produced at channel index 0.
+		 * If false, then each channel is analyzed separately. 
+		 */
+		bool bDownmixToMono;
+
+		FConstantQNRTSettings()
+		:	AnalysisPeriod(0.01f)
+		,	bDownmixToMono(false)
+		{}
+	};
+
+
+	/** FConstantQFrame
+	 * 
+	 *  Contains Constant Q data relating to one audio window.
+	 */
+	struct AUDIOSYNESTHESIACORE_API FConstantQFrame
+	{
+		/** Audio channel which produced the data. */
+		int32 Channel;
+		
+		/** Timestamp in seconds referring to the center of the audio window. */
+		float Timestamp;
+
+		/** Output spectral data */
+		TArray<float> Spectrum;
+
+		FConstantQFrame()
+		:	Channel(0)
+		,	Timestamp(0.f)
+		{}
+
+		FConstantQFrame(int32 InChannelIndex, float InTimestamp, TArrayView<const float> InSpectrum)
+		:	Channel(InChannelIndex)
+		,	Timestamp(InTimestamp)
+		,	Spectrum(InSpectrum.GetData(), InSpectrum.Num())
+		{}
+	};
+
+	/** Serialize FConstantQFrame */
+	AUDIOSYNESTHESIACORE_API FArchive &operator <<(FArchive& Ar, FConstantQFrame& Frame);
+
+	/** FConstantQNRTResult
+	 *
+	 *  FConstantQNRTResult is a container for the output of the FConstantQNRTWorker.
+	 */
+	class AUDIOSYNESTHESIACORE_API FConstantQNRTResult : public IAnalyzerNRTResult
+	{
+	public:
+
+		FConstantQNRTResult();
+
+		/** Serialize or unserialize object */
+		virtual void Serialize(FArchive& Archive) override;
+
+		/** Add a single frame of CQT data */
+		void AddFrame(int32 InChannelIndex, float InTimestamp, TArrayView<const float> InSpectrum);
+
+		/** Returns true if this result contains data for the given channel index. */
+		bool ContainsChannel(int32 InChannelIndex) const;
+
+		/** Retrieve the array of frames for a single channel of audio. */
+		const TArray<FConstantQFrame>& GetFramesForChannel(int32 InChannelIndex) const;
+
+		/** Retrieve the difference between the maximum and minimum value in the spectrum. */
+		FFloatInterval GetChannelConstantQInterval(int32 InChannelIdx) const;
+
+		/** Retrieve an array of channel indices which exist in this result. */
+		void GetChannels(TArray<int32>& OutChannels) const;
+
+		/** Returns the duration of the analyzed audio in seconds */
+		virtual float GetDurationInSeconds() const override;
+
+		/** Sets the duration of the analyzed audio in seconds */
+		void SetDurationInSeconds(float InDuration); 
+
+	 	/** Returns true if FConstantQFrame arrays are sorted in chronologically ascending order via their timestamp.  */
+		bool IsSortedChronologically() const;
+
+		/** Sorts FConstantQFrame arrays in chronologically ascnding order via their timestamp.  */
+		void SortChronologically();
+
+
+	private:
+
+		float DurationInSeconds;
+
+		TMap<int32, TArray<FConstantQFrame> > ChannelCQTFrames;
+
+		TMap<int32, FFloatInterval> ChannelCQTIntervals;
+
+		bool bIsSortedChronologically;
+	};
+
+	/** FConstantQNRTWorker
+	 *
+	 *  FConstantQNRTWorker computes a FConstantQNRTResult from audio samples.
+	 */
+	class AUDIOSYNESTHESIACORE_API FConstantQNRTWorker : public IAnalyzerNRTWorker
+	{
+	public:
+
+		/** Constructor
+		 *
+		 * InParams are the parameters which describe characteristics of the input audio.
+		 * InAnalyzerSettings are the settings which control various aspects of the algorithm.
+		 */
+		FConstantQNRTWorker(const FAnalyzerNRTParameters& InParams, const FConstantQNRTSettings& InAnalyzerSettings);
+
+		/**
+		 *  Analyze audio and put results into results pointer.
+		 *
+		 *  InAudio is an array view of audio.
+		 *  OutResult is a pointer to a valid FConstantQNRTResult
+		 */
+		virtual void Analyze(TArrayView<const float> InAudio, IAnalyzerNRTResult* OutResult) override;
+
+		/** 
+		 *  Call when analysis of audio asset is complete. 
+		 *
+		 *  OutResult must be a pointer to a valid FConstantQNRTResult. 
+		 */
+		virtual void Finalize(IAnalyzerNRTResult* OutResult) override;
+
+	private:
+
+		/** Analyze audio with multiple channels interleaved. */
+		void AnalyzeMultichannel(TArrayView<const float> InAudio, IAnalyzerNRTResult* OutResult, bool bDoFlush);
+
+		/** Analyze a single window of audio from a single channel */
+		void AnalyzeWindow(const AlignedFloatBuffer& InWindow, int32 InChannelIndex, FConstantQNRTResult& OutResult);
+
+		int32 NumFrames;
+		int32 NumChannels;
+		int32 NumBuffers;
+		float SampleRate;
+		int32 NumHopFrames;
+		int32 NumHopSamples;
+		int32 NumWindowFrames;
+		int32 NumWindowSamples;
+
+		float MonoScaling;
+
+		AlignedFloatBuffer HopBuffer;
+		AlignedFloatBuffer ChannelBuffer;
+		AlignedFloatBuffer MonoBuffer;
+
+		TArray<float> CQTSpectrum;
+
+		TUniquePtr<TSlidingBuffer<float> > SlidingBuffer;
+		TUniquePtr<FConstantQAnalyzer> ConstantQAnalyzer;
+
+		bool bDownmixToMono;
+	};
+
+	/** FConstantQNRTFactory
+	 *  
+	 *  Factory for creating FConstantQNRT workers and results
+	 */
+	class AUDIOSYNESTHESIACORE_API FConstantQNRTFactory : public IAnalyzerNRTFactory
+	{
+	public:
+
+		/** Name of this analyzer type. */
+		virtual FName GetName() const override;
+
+		/** Human readable name of this analyzer. */
+		virtual FString GetTitle() const override;
+
+		/** Create a new FConstantQNRTResult. */
+		virtual TUniquePtr<IAnalyzerNRTResult> NewResult() const override;
+
+		/** Create a new FConstantQNRTWorker 
+		 *
+		 *  InSettings must be a pointer to FConstantQNRTSetting
+		 */
+		virtual TUniquePtr<IAnalyzerNRTWorker> NewWorker(const FAnalyzerNRTParameters& InParams, const IAnalyzerNRTSettings* InSettings) const override;
+	};
+}
diff --git a/Engine/Plugins/Runtime/AudioSynesthesia/Source/AudioSynesthesiaCore/Public/InterpolateSorted.h b/Engine/Plugins/Runtime/AudioSynesthesia/Source/AudioSynesthesiaCore/Public/InterpolateSorted.h
new file mode 100644
--- /dev/null
+++ b/Engine/Plugins/Runtime/AudioSynesthesia/Source/AudioSynesthesiaCore/Public/InterpolateSorted.h
@@ -0,0 +1,48 @@
+// Copyright Epic Games, Inc. All Rights Reserved.
+
+#pragma once
+
+namespace Audio
+{
+	/** GetInterpolationParametersAtTimestamp
+	 *
+	 *  Calculates the indices and alpha parameter for interpolating between to elements in the array.
+	 *
+	 *  ElementType must be a type with a "float Timestamp;" member.
+	 *  InArray is a array of ElementTypes which are sorted in ascending order via the Timestamp member.
+	 *  InTimestamp is the query timestamp.
+	 *
+	 *  OutLowerIndex is populated with the index of the last element in InArray with a Timestamp less than or equal to InTimestamp. On error, this is set to INDEX_NONE
+	 *  OutUpperIndex is populated with the index of the first element in InArray with a Timestamp greater than InTimestamp. On error, this is set to INDEX_NONE
+	 *  OutAlpha is a number between 0.f and 1.f which describes the blending between the lower and upper indices. 0.f refers OutLowerIndex and 1.f to OutUpperIndex.
+	 */
+	template<typename ElementType>
+	void GetInterpolationParametersAtTimestamp(TArrayView<const ElementType> InArray, float InTimestamp, int32& OutLowerIndex, int32& OutUpperIndex, float& OutAlpha)
+	{
+		if (0 == InArray.Num())
+		{
+			OutLowerIndex = INDEX_NONE;
+			OutUpperIndex = INDEX_NONE;
+			OutAlpha = 1.f;
+			return;
+		}
+		
+		// Get upper and lower array indices with FLoudnessData closest to the query timestamp.
+		OutUpperIndex = Algo::UpperBoundBy(InArray, InTimestamp, [](const ElementType& Datum) { return Datum.Timestamp; });
+		OutLowerIndex = OutUpperIndex - 1;
+
+		// Ensure indices are within range.
+		int32 MaxIdx = InArray.Num() - 1;
+		OutLowerIndex = FMath::Clamp(OutLowerIndex, 0, MaxIdx);
+		OutUpperIndex = FMath::Clamp(OutUpperIndex, 0, MaxIdx);
+
+		const ElementType& LowerDatum = InArray[OutLowerIndex];
+		const ElementType& UpperDatum = InArray[OutUpperIndex];
+
+		// Determine time difference and linear interpolation factor from data.
+		const float TimeDelta = FMath::Max(UpperDatum.Timestamp - LowerDatum.Timestamp, SMALL_NUMBER);
+		OutAlpha = (InTimestamp - LowerDatum.Timestamp) / TimeDelta;
+		OutAlpha = FMath::Clamp(OutAlpha, 0.0f, 1.0f);
+	}
+}
+
diff --git a/Engine/Plugins/Runtime/AudioSynesthesia/Source/AudioSynesthesiaCore/Public/LoudnessAnalyzer.h b/Engine/Plugins/Runtime/AudioSynesthesia/Source/AudioSynesthesiaCore/Public/LoudnessAnalyzer.h
new file mode 100644
--- /dev/null
+++ b/Engine/Plugins/Runtime/AudioSynesthesia/Source/AudioSynesthesiaCore/Public/LoudnessAnalyzer.h
@@ -0,0 +1,172 @@
+// Copyright Epic Games, Inc. All Rights Reserved.
+
+#pragma once
+
+#include "DSP/BufferVectorOperations.h"
+#include "DSP/AudioFFT.h"
+
+namespace Audio
+{
+    /** Equal loudness curves from: ANSI Standards S1.4-1983 and S1.42-2001 */
+	enum class ELoudnessCurveType : uint8
+	{
+		A,
+		B,
+		C,
+		D,
+		None
+	};
+
+	/** Generate equal loudness weight for magnitude spectrum using curve A */
+    AUDIOSYNESTHESIACORE_API float GetEqualLoudnessAWeigthForFrequency(const float Freq);
+
+	/** Generate equal loudness weight for magnitude spectrum using curve B */
+    AUDIOSYNESTHESIACORE_API float GetEqualLoudnessBWeigthForFrequency(const float Freq);
+
+	/** Generate equal loudness weight for magnitude spectrum using curve C */
+    AUDIOSYNESTHESIACORE_API float GetEqualLoudnessCWeigthForFrequency(const float Freq);
+
+	/** Generate equal loudness weight for magnitude spectrum using curve D */
+    AUDIOSYNESTHESIACORE_API float GetEqualLoudnessDWeigthForFrequency(const float Freq);
+
+	/** Generate equal loudness weight for magnitude spectrum using curve None */
+    AUDIOSYNESTHESIACORE_API float GetEqualLoudnessNoneWeigthForFrequency(const float Freq);
+
+	/** Settings for FLoudnessAnalyzer */
+    struct AUDIOSYNESTHESIACORE_API FLoudnessAnalyzerSettings
+    {
+		/** Size of FFT describe in number of samples */
+        int32 FFTSize;
+
+		/** Type of window to be applied to input audio */
+        EWindowType WindowType;
+
+		/** Type of equal loudness curve to weight power spectrum */
+        ELoudnessCurveType LoudnessCurveType;
+
+		/** Minimum frequency for analyzing audio */
+        float MinAnalysisFrequency;
+
+		/** Maximum frequency for analyzing audio */
+        float MaxAnalysisFrequency;
+
+        /** Default settings */
+		FLoudnessAnalyzerSettings()
+			: FFTSize(8192)
+			, WindowType(EWindowType::Blackman)
+			, LoudnessCurveType(ELoudnessCurveType::D)
+			, MinAnalysisFrequency(20.0f)
+			, MaxAnalysisFrequency(10000.0f)
+		{}
+    };
+
+	/** FLoudnessAnalyzer
+	 *
+	 * FLoudnessAnalyzer will calculate perceptual energy or loudness one window at a time. 
+	 * This analyzer only handles mono audio. See FMultichannelLoudnessAnalyzer for multichannel
+	 * support.
+	 */
+    class AUDIOSYNESTHESIACORE_API FLoudnessAnalyzer
+    {
+        public:
+			/** Construct analyzer */
+            FLoudnessAnalyzer(float InSampleRate, const FLoudnessAnalyzerSettings& InSettings);
+
+			/**
+			 * Calculats the perceptual energy for the input samples.  The number of 
+			 * input samples must be equal to the FFTSize supplied in FLoudnessAnalyzerSettings.
+			 */
+            float CalculatePerceptualEnergy(TArrayView<const float> InSampleView);
+
+			/**
+			 * Calculats the loudness for the input samples.  The number of  input samples must 
+			 * be equal to the FFTSize supplied in FLoudnessAnalyzerSettings.
+			 */
+            float CalculateLoudness(TArrayView<const float> InSampleView);
+
+			/**
+			 * Convert from perceptual energy to loudness [Loudness = 10 * log_10(Energy)]
+			 */
+			static float ConvertPerceptualEnergyToLoudness(float InPerceptualEnergy);
+
+			/**
+			 * Return const reference to settings used inside this analyzer.
+			 */
+            const FLoudnessAnalyzerSettings& GetSettings() const;
+        private:
+
+            FLoudnessAnalyzerSettings Settings;
+
+			/** Scale factor used to normalize between differing window sizes and fft sizes */
+            float EnergyScale;
+
+            int32 MinFreqIndex;
+            int32 MaxFreqIndex;
+
+			/** Number of FFTFrequency bins worth looking at. */
+			int32 FFTFreqSize;
+
+            FWindow Window;
+
+            AlignedFloatBuffer WindowedSamples;
+            AlignedFloatBuffer FFTOutputRealData;
+            AlignedFloatBuffer FFTOutputImagData;
+
+            AlignedFloatBuffer CurveWeights;
+            AlignedFloatBuffer PowerSpectrumBuffer;
+    };
+
+
+	/** FMultichannelLoudnessAnalyzer
+	 *
+	 * FMultichannelLoudnessAnalyzer will calculate perceptual energy or loudness one window 
+	 * at a time. This module can handle one or more channels at a time. It expects 
+	 * interleaved audio samples.
+	 */
+    class AUDIOSYNESTHESIACORE_API FMultichannelLoudnessAnalyzer 
+    {
+        public:
+			/** Construct analyzer. */
+            FMultichannelLoudnessAnalyzer(float InSampleRate, const FLoudnessAnalyzerSettings& InSettings);
+
+			/** 
+			 * Calculate perceptual energy. 
+			 *
+			 * InSampleView		  InSampleView is expected to contain interleaved audio with 
+			 *                    (InSettings.FFTSize * InNumChannels) number of samples.
+			 *
+			 * InNumChannels      Number of input channels.
+			 *
+			 * OutChannelEnergies Will be populated with energies on a per channel basis.
+			 *
+			 * The overall perceptual energy is returned.
+			 */
+            float CalculatePerceptualEnergy(TArrayView<const float> InSampleView, const int32 InNumChannels, TArray<float>& OutChannelEnergies); 
+
+			/** 
+			 * Calculate loudness
+			 *
+			 * InSampleView		  InSampleView is expected to contain interleaved audio with 
+			 *                    (InSettings.FFTSize * InNumChannels) number of samples.
+			 *
+			 * InNumChannels      Number of input channels.
+			 *
+			 * OutChannelLoudness Will be populated with energies on a per channel basis.
+			 *
+			 * The overall loudness is returned.
+			 */
+            float CalculateLoudness(TArrayView<const float> InSampleView, const int32 InNumChannels, TArray<float>& OutChannelLoudness);
+
+			/**
+			 * Return const reference to settings used inside this analyzer.
+			 */
+            const FLoudnessAnalyzerSettings& GetSettings() const;
+
+        private:
+
+            AlignedFloatBuffer MonoBuffer;
+            TArray<float> ChannelWeights;
+            FLoudnessAnalyzer Analyzer;
+    };
+}
+
diff --git a/Engine/Plugins/Runtime/AudioSynesthesia/Source/AudioSynesthesiaCore/Public/LoudnessNRTFactory.h b/Engine/Plugins/Runtime/AudioSynesthesia/Source/AudioSynesthesiaCore/Public/LoudnessNRTFactory.h
new file mode 100644
--- /dev/null
+++ b/Engine/Plugins/Runtime/AudioSynesthesia/Source/AudioSynesthesiaCore/Public/LoudnessNRTFactory.h
@@ -0,0 +1,180 @@
+// Copyright Epic Games, Inc. All Rights Reserved.
+
+#pragma once
+
+#include "LoudnessAnalyzer.h"
+#include "IAudioAnalyzerNRTInterface.h"
+#include "DSP/SlidingWindow.h"
+
+namespace Audio
+{
+	/** FLoudnessNRTSettings
+	 *
+	 * Contains settings for loudness analyzer.
+	 */
+	class AUDIOSYNESTHESIACORE_API FLoudnessNRTSettings : public IAnalyzerNRTSettings, public FLoudnessAnalyzerSettings
+	{
+		public:
+			/** Number of seconds between loudness measurements */
+			float AnalysisPeriod;
+
+			FLoudnessNRTSettings()
+			:	AnalysisPeriod(0.01f)
+			{}
+	};
+
+
+	/** FLoudnessNRTResult
+	 *
+	 * Holds the loundess values per a time step.
+	 */
+	struct FLoudnessDatum
+	{
+		int32 Channel = 0;
+		float Timestamp = 0.f;
+		float Energy = 0.f;
+		float Loudness = 0.f;
+	};
+	/** De/Serialize single loudness datum into archive. */
+	AUDIOSYNESTHESIACORE_API FArchive &operator <<(FArchive& Ar, FLoudnessDatum& Datum);
+
+	/** FLoudnessNRTResult
+	 *
+	 * FLoudnessNRTResult contains the temporal evolution of loudness.
+	 */
+	class AUDIOSYNESTHESIACORE_API FLoudnessNRTResult : public IAnalyzerNRTResult
+	{
+	public:
+		/** ChannelIndexOverall
+		 *
+		 * Denotes the overall loudness channel index as opposed individual channel indices.
+		 */
+		static const int32 ChannelIndexOverall;
+
+		FLoudnessNRTResult();
+
+		/**
+		 * Defines how to serialize result.
+		 */
+		virtual void Serialize(FArchive& Archive) override;
+
+		/**
+		 * Appends an FLoudnessDatum to the container.
+		 */
+		void Add(const FLoudnessDatum& InDatum);
+
+		/** Returns true if this object data for the given channel index */
+		bool ContainsChannel(int32 InChannelIndex) const;
+
+		/**
+		 * Returns const reference to FLoudnessDatum array for individual channel.
+		 */
+		const TArray<FLoudnessDatum>& GetChannelLoudnessArray(int32 ChannelIdx) const;
+
+		/**
+		 * Returns const reference to FLoudnessDatum array associated with overall loudness.
+		 */
+		const TArray<FLoudnessDatum>& GetLoudnessArray() const;
+
+		/**
+		 * Returns range in dB of overall loudness result given the noise floor.
+		 */
+		float GetLoudnessRange(float InNoiseFloor) const;
+
+
+		/** Returns range in dB of loudness result given the noise floor. */
+		float GetChannelLoudnessRange(int32 InChannelIdx, float InNoiseFloor) const;
+
+		/** Returns the channel indices availabe in result. */
+		void GetChannels(TArray<int32>& OutChannels) const;
+
+		/** Gets the duration of the anlayzed audio. */
+		virtual float GetDurationInSeconds() const override;
+
+		/** Sets the duration of the anlayzed audio. */
+		void SetDurationInSeconds(float InDuration);
+
+		/**
+		 * Returns true if FLoudnessDatum arrays are sorted in chronologically ascending order via their timestamp.
+		 */
+		bool IsSortedChronologically() const;
+
+		/**
+		 * Sorts FLoudnessDatum arrays in chronologically ascnding order via their timestamp.
+		 */
+		void SortChronologically();
+
+	private:
+		float DurationInSeconds;
+
+		TMap<int32, TArray<FLoudnessDatum> > ChannelLoudnessArrays;
+
+		TMap<int32, FFloatInterval> ChannelLoudnessIntervals;
+
+		bool bIsSortedChronologically;
+	};
+
+	/** FLoudnessNRTWorker
+	 *
+	 * FLoudnessNRTWorker performs loudness analysis on input sample buffers.
+	 */
+	class AUDIOSYNESTHESIACORE_API FLoudnessNRTWorker : public IAnalyzerNRTWorker
+	{
+	public:
+		/** Construct a worker */
+		FLoudnessNRTWorker(const FAnalyzerNRTParameters& InParams, const FLoudnessNRTSettings& InAnalyzerSettings);
+
+		/**
+		 * Analyzes input sample buffer and updates result. 
+		 */
+		virtual void Analyze(TArrayView<const float> InAudio, IAnalyzerNRTResult* OutResult) override;
+
+		/**
+		 * Call when all audio data has been analyzed. 
+		 */
+		virtual void Finalize(IAnalyzerNRTResult* OutResult) override;
+
+	private:
+
+		// Analyze a single window
+		void AnalyzeWindow(TArrayView<const float> InWindow, FLoudnessNRTResult& OutResult);
+
+		int32 NumChannels;
+		int32 NumAnalyzedBuffers;
+		int32 NumHopFrames;
+		int32 NumFrames;
+
+		float SampleRate;
+
+		TArray<float> InternalWindow;
+
+		TUniquePtr<TSlidingBuffer<float>> InternalBuffer;
+
+		TUniquePtr<FMultichannelLoudnessAnalyzer> Analyzer;
+	};
+
+	/** FLoudnessNRTFactory
+	 *
+	 * Defines the LoudnessNRT analyzer and creates related classes.
+	 */
+	class AUDIOSYNESTHESIACORE_API FLoudnessNRTFactory : public IAnalyzerNRTFactory
+	{
+		public:
+
+		/** Name of specific analyzer type. */
+		virtual FName GetName() const override;
+
+		/** Human readable name of analyzer. */
+		virtual FString GetTitle() const override;
+
+		/** Creates a new FLoudnessNRTResult */
+		virtual TUniquePtr<IAnalyzerNRTResult> NewResult() const override;
+
+		/** 
+		 * Creates a new FLoudnessNRTWorker. This expects IAnalyzerNRTSettings to be a valid pointer to
+		 * a FLoudnessNRTSettings object.
+		 */
+		virtual TUniquePtr<IAnalyzerNRTWorker> NewWorker(const FAnalyzerNRTParameters& InParams, const IAnalyzerNRTSettings* InSettings) const override;
+	};
+}
+
diff --git a/Engine/Plugins/Runtime/AudioSynesthesia/Source/AudioSynesthesiaCore/Public/OnsetAnalyzer.h b/Engine/Plugins/Runtime/AudioSynesthesia/Source/AudioSynesthesiaCore/Public/OnsetAnalyzer.h
new file mode 100644
--- /dev/null
+++ b/Engine/Plugins/Runtime/AudioSynesthesia/Source/AudioSynesthesiaCore/Public/OnsetAnalyzer.h
@@ -0,0 +1,106 @@
+// Copyright Epic Games, Inc. All Rights Reserved.
+
+#pragma once
+
+#include "PeakPicker.h"
+#include "DSP/MelScale.h"
+#include "DSP/SlidingWindow.h"
+#include "DSP/AudioFFT.h"
+
+namespace Audio
+{
+	/** Onset strength settings */
+	struct AUDIOSYNESTHESIACORE_API FOnsetStrengthSettings
+	{
+		// Number of frames between strength windows
+		int32 NumHopFrames = 1024;
+
+		// Number of frames used to represent one audio window
+		int32 NumWindowFrames = 4096;
+
+		// Number of windows to lag when calculating window difference.
+		int32 ComparisonLag = 4;
+
+		// Size of FFT. Must be greater than or equal to NumWindowFrames.
+		int32 FFTSize = 4096;
+
+		// Type of analysis window to apply to audio.
+		EWindowType WindowType = EWindowType::Blackman;
+
+		// Decibel level considered to be silence.
+		float NoiseFloorDb = -60.f;
+
+		// Settings for converting FFT to mel spectrum.
+		FMelSpectrumKernelSettings MelSettings;
+	};
+
+	/** FOnsetStrengthAnalyzer
+	 *
+	 * FOnsetStrengthAnalyzer calculates the onset strength from audio. 
+	 *
+	 * Onset strength is calculated as the half wave rectificed difference between two spectral frames.
+	 * This onset strength onset analyzer uses the following approach.
+	 *
+	 * [audio]->[fft]->[mel spectrum]->[diff]->[half wave rectify]->[mean]->[onset strength]
+	 *                       |           |
+	 *                       >-[lag]-----|
+	 */
+	class AUDIOSYNESTHESIACORE_API FOnsetStrengthAnalyzer
+	{
+		public:
+			FOnsetStrengthAnalyzer(const FOnsetStrengthSettings& InSettings, float InSampleRate);
+
+			/** Calculates onset strengths from audio and fills OutOnsetStrengths with generated onset strengths */
+			void CalculateOnsetStrengths(TArrayView<const float> InSamples, TArray<float>& OutOnsetStrengths);
+
+			/** Call when done processing audio for an audio analyzer */
+			void FlushAudio(TArray<float>& OutEnvelopeStrengths);
+			
+			/** Call to reset internal counters and lag spectra. */
+			void Reset();
+
+			/** Converts an onset strength index into a timestamp. */
+			static float GetTimestampForIndex(const FOnsetStrengthSettings& InSettings, float InSampleRate, int32 InIndex);
+
+		private:
+
+			void AnalyzeAudio(TArrayView<const float> InSamples, TArray<float>& OutEnvelopeStrengths, bool bDoFlush = false);
+
+			// Processes single audio window
+			float GetNextOnsetStrength(TArrayView<const float> InSamples);
+
+			FOnsetStrengthSettings Settings;
+			int32 LagSpectraIndex;
+
+			TSlidingBuffer<float> SlidingBuffer;
+			AlignedFloatBuffer WorkingBuffer;
+			AlignedFloatBuffer WindowedSamples;
+			AlignedFloatBuffer FFTOutputRealData;
+			AlignedFloatBuffer FFTOutputImagData;
+			AlignedFloatBuffer FFTSpectrumBuffer;
+
+			// Array of arrays to hold lag spectra.
+			TArray< TArray<float> > PreviousSpectra;
+
+			TArray<float> MelSpectrum;
+			TArray<float> SpectrumDifference;
+
+			FWindow Window;
+			TUniquePtr<FContiguousSparse2DKernelTransform> MelTransform;
+	};
+
+	/** 
+	 * Extract onset indicies from an onset strength envelope.
+	 *
+	 * InSettings is the peak picker settings used to pick onset peaks from an onset strength envelope.
+	 * InOnsetEnvelop is an onset strength envelope generated using the FOnsetStrengthAnalyzer.
+	 * OutOnsetIndices contains the indices where onsets occur.
+	 */
+	AUDIOSYNESTHESIACORE_API void OnsetExtractIndices(const FPeakPickerSettings& InSettings, TArrayView<const float> InOnsetEnvelope, TArray<int32>& OutOnsetIndices);
+
+	/**
+	 * Backtracks onset indices to the beginning of the onset attack. This is useful when generating 
+	 * splice points which retain the onset attack. 
+	 */
+	AUDIOSYNESTHESIACORE_API void OnsetBacktrackIndices(TArrayView<const float> InOnsetEnvelope, TArrayView<const int32> InOnsetIndices, TArray<int32>& OutBacktrackedOnsetIndices);
+}
diff --git a/Engine/Plugins/Runtime/AudioSynesthesia/Source/AudioSynesthesiaCore/Public/OnsetNRTFactory.h b/Engine/Plugins/Runtime/AudioSynesthesia/Source/AudioSynesthesiaCore/Public/OnsetNRTFactory.h
new file mode 100644
--- /dev/null
+++ b/Engine/Plugins/Runtime/AudioSynesthesia/Source/AudioSynesthesiaCore/Public/OnsetNRTFactory.h
@@ -0,0 +1,205 @@
+// Copyright Epic Games, Inc. All Rights Reserved.
+
+#pragma once
+
+#include "IAudioAnalyzerNRTInterface.h"
+#include "OnsetAnalyzer.h"
+#include "PeakPicker.h"
+
+namespace Audio
+{
+	/** FOnsetNRTSettings
+	 *
+	 *  Settings for the Constant Q Non-Real-Time Factory.
+	 */
+	class AUDIOSYNESTHESIACORE_API FOnsetNRTSettings : public IAnalyzerNRTSettings
+	{
+	public:
+
+		FOnsetStrengthSettings OnsetStrengthSettings;
+
+		FPeakPickerSettings PeakPickerSettings;
+
+		/** 
+		 * If true, all channels are mixed together with equal power before analysis. Only one channel
+		 * is produced at channel index 0.
+		 * If false, then each channel is analyzed separately. 
+		 */
+		bool bDownmixToMono;
+
+		FOnsetNRTSettings()
+		:	bDownmixToMono(false)
+		{}
+	};
+
+
+	/** FOnset
+	 * 
+	 *  Contains data relating to one onset.
+	 */
+	struct AUDIOSYNESTHESIACORE_API FOnset
+	{
+		/** Audio channel which produced the data. */
+		int32 Channel;
+		
+		/** Timestamp in seconds referring to the center of the audio window. */
+		float Timestamp;
+
+		float Strength;
+
+		FOnset()
+		:	Channel(0)
+		,	Timestamp(0.f)
+		,	Strength(0.f)
+		{}
+
+		FOnset(int32 InChannelIndex, float InTimestamp, float InStrength)
+		:	Channel(InChannelIndex)
+		,	Timestamp(InTimestamp)
+		,	Strength(InStrength)
+		{}
+	};
+
+	/** Serialize FOnset */
+	AUDIOSYNESTHESIACORE_API FArchive &operator <<(FArchive& Ar, FOnset& Frame);
+
+	/** FOnsetNRTResult
+	 *
+	 *  FOnsetNRTResult is a container for the output of the FOnsetNRTWorker.
+	 */
+	class AUDIOSYNESTHESIACORE_API FOnsetNRTResult : public IAnalyzerNRTResult
+	{
+	public:
+
+		FOnsetNRTResult();
+
+		/** Serialize or unserialize object */
+		virtual void Serialize(FArchive& Archive) override;
+
+		/** Add a single frame of CQT data */
+		void AddOnset(int32 InChannelIndex, float InTimestamp, float InStrength);
+
+		/** Adds an empty channel. Useful if a channel exists without any onsets. */
+		void AddChannel(int32 InChannelIndex);
+
+		/** Returns true if this result contains the channel index. */
+		bool ContainsChannel(int32 InChannelIndex) const;
+
+		/** Retrieve the array of frames for a single channel of audio. */
+		const TArray<FOnset>& GetOnsetsForChannel(int32 InChannelIndex) const;
+
+		/** Retrieve the interval of the maximum and minimum strengths of all onsets in channel. */
+		FFloatInterval GetChannelOnsetInterval(int32 InChannelIdx) const;
+
+		/** Retrieve the channels available in the result. */
+		void GetChannels(TArray<int32>& OutChannels) const;
+
+		/** Returns the duration of the analyzed audio in seconds */
+		virtual float GetDurationInSeconds() const override;
+
+		/** Sets the duration of the analyzed audio in seconds */
+		void SetDurationInSeconds(float InDuration); 
+
+	 	/** Returns true if FOnset arrays are sorted in chronologically ascending order via their timestamp.  */
+		bool IsSortedChronologically() const;
+
+		/** Sorts FOnset arrays in chronologically ascnding order via their timestamp.  */
+		void SortChronologically();
+
+	private:
+
+		float DurationInSeconds;
+
+		TMap<int32, TArray<FOnset> > ChannelOnsets;
+
+		TMap<int32, FFloatInterval> ChannelOnsetStrengthIntervals;
+
+		bool bIsSortedChronologically;
+	};
+
+	/** FOnsetNRTWorker
+	 *
+	 *  FOnsetNRTWorker computes a FOnsetNRTResult from audio samples.
+	 */
+	class AUDIOSYNESTHESIACORE_API FOnsetNRTWorker : public IAnalyzerNRTWorker
+	{
+	public:
+
+		/** Constructor
+		 *
+		 * InParams are the parameters which describe characteristics of the input audio.
+		 * InAnalyzerSettings are the settings which control various aspects of the algorithm.
+		 */
+		FOnsetNRTWorker(const FAnalyzerNRTParameters& InParams, const FOnsetNRTSettings& InAnalyzerSettings);
+
+		/**
+		 *  Analyze audio and put results into results pointer.
+		 *
+		 *  InAudio is an array view of audio.
+		 *  OutResult is a pointer to a valid FOnsetNRTResult
+		 */
+		virtual void Analyze(TArrayView<const float> InAudio, IAnalyzerNRTResult* OutResult) override;
+
+		/** 
+		 *  Call when analysis of audio asset is complete. 
+		 *
+		 *  OutResult must be a pointer to a valid FOnsetNRTResult. 
+		 */
+		virtual void Finalize(IAnalyzerNRTResult* OutResult) override;
+
+	private:
+
+		/** Analyze audio with multiple channels interleaved. */
+		void AnalyzeMultichannel(TArrayView<const float> InAudio, IAnalyzerNRTResult* OutResult, bool bDoFlush);
+
+		void AnalyzeWindow(const AlignedFloatBuffer& InWindow, int32 InChannelIndex, FOnsetNRTResult& OutResult);
+
+		void ExtractOnsetsFromOnsetStrengths(FOnsetNRTResult& OutResult);
+
+		void Reset();
+
+		FOnsetNRTSettings Settings;
+
+		int32 NumFrames;
+		int32 NumChannels;
+		float SampleRate;
+		int32 NumWindowFrames;
+		int32 NumWindowSamples;
+		float MonoScaling;
+
+		AlignedFloatBuffer HopBuffer;
+		AlignedFloatBuffer ChannelBuffer;
+		AlignedFloatBuffer MonoBuffer;
+
+		TUniquePtr<TSlidingBuffer<float> > SlidingBuffer;
+
+		TArray< TUniquePtr<FOnsetStrengthAnalyzer> > OnsetStrengthAnalyzers;
+
+		TArray< TArray<float> > OnsetStrengths;
+
+	};
+
+	/** FOnsetNRTFactory
+	 *  
+	 *  Factory for creating FOnsetNRT workers and results
+	 */
+	class AUDIOSYNESTHESIACORE_API FOnsetNRTFactory : public IAnalyzerNRTFactory
+	{
+	public:
+
+		/** Name of this analyzer type. */
+		virtual FName GetName() const override;
+
+		/** Human readable name of this analyzer. */
+		virtual FString GetTitle() const override;
+
+		/** Create a new FOnsetNRTResult. */
+		virtual TUniquePtr<IAnalyzerNRTResult> NewResult() const override;
+
+		/** Create a new FOnsetNRTWorker 
+		 *
+		 *  InSettings must be a pointer to FOnsetNRTSetting
+		 */
+		virtual TUniquePtr<IAnalyzerNRTWorker> NewWorker(const FAnalyzerNRTParameters& InParams, const IAnalyzerNRTSettings* InSettings) const override;
+	};
+}
diff --git a/Engine/Plugins/Runtime/AudioSynesthesia/Source/AudioSynesthesiaCore/Public/PeakPicker.h b/Engine/Plugins/Runtime/AudioSynesthesia/Source/AudioSynesthesiaCore/Public/PeakPicker.h
new file mode 100644
--- /dev/null
+++ b/Engine/Plugins/Runtime/AudioSynesthesia/Source/AudioSynesthesiaCore/Public/PeakPicker.h
@@ -0,0 +1,32 @@
+// Copyright Epic Games, Inc. All Rights Reserved.
+
+#pragma once
+
+#include "CoreMinimal.h"
+
+namespace Audio
+{
+	struct FPeakPickerSettings
+	{
+		int32 NumPreMax = 1440;
+		int32 NumPostMax = 1;
+		int32 NumPreMean = 4800;
+		int32 NumPostMean = 4801;
+		int32 NumWait = 1440;
+		float MeanDelta = 0.07f;
+	};
+
+	class AUDIOSYNESTHESIACORE_API FPeakPicker
+	{
+		public:
+			FPeakPicker(const FPeakPickerSettings& InSettings);
+
+			void PickPeaks(TArrayView<const float> InView, TArray<int32>& OutPeakIndices);
+
+		private:
+			FPeakPickerSettings Settings;
+
+
+	};
+}
+
diff --git a/Engine/Source/Runtime/AudioAnalyzer/Private/IAudioAnalyzerNRTInterface.cpp b/Engine/Source/Runtime/AudioAnalyzer/Private/IAudioAnalyzerNRTInterface.cpp
--- a/Engine/Source/Runtime/AudioAnalyzer/Private/IAudioAnalyzerNRTInterface.cpp
+++ b/Engine/Source/Runtime/AudioAnalyzer/Private/IAudioAnalyzerNRTInterface.cpp
@@ -2,39 +2,7 @@
 
 #include "IAudioAnalyzerNRTInterface.h"
 #include "CoreMinimal.h"
-#include "Serialization/BitReader.h"
-#include "Serialization/BitWriter.h"
 
 namespace Audio
 {
-	void IAnalyzerNRTResult::CopyFrom(IAnalyzerNRTResult* SourceResult)
-	{
-		FBitWriter SerializedResult;
-		SourceResult->Serialize(SerializedResult);
-		FBitReader DeserializedResult = FBitReader(SerializedResult.GetData(), SerializedResult.GetNumBytes() * 8);
-		Serialize(DeserializedResult);
-	}
-	
-	FString IAnalyzerNRTResult::ToString() const
-	{
-		return FString(TEXT("To use, override ToString() in this implementation of IAnalyzerNRTResult."));
-	}
-
-	FName IAnalyzerNRTFactory::GetModularFeatureName()
-	{
-		static FName AudioExtFeatureName = FName(TEXT("AudioAnalyzerNRTPlugin"));
-		return AudioExtFeatureName;
-	}
-
-	// This is used to identify analyzers created with this factory.
-	FName IAnalyzerNRTFactory::GetName() const
-	{
-		static FName DefaultName(TEXT("UnnamedAudioAnalyzerNRT"));
-		return DefaultName;
-	}
-
-	FString IAnalyzerNRTFactory::GetTitle() const
-	{
-		return FString("Unnamed Non Real-Time Audio Analyzer.");
-	}
 }
diff --git a/Engine/Source/Runtime/AudioAnalyzer/Public/IAudioAnalyzerNRTInterface.h b/Engine/Source/Runtime/AudioAnalyzer/Public/IAudioAnalyzerNRTInterface.h
--- a/Engine/Source/Runtime/AudioAnalyzer/Public/IAudioAnalyzerNRTInterface.h
+++ b/Engine/Source/Runtime/AudioAnalyzer/Public/IAudioAnalyzerNRTInterface.h
@@ -6,6 +6,8 @@
 
 #include "CoreMinimal.h"
 #include "Features/IModularFeature.h"
+#include "Serialization/BitReader.h"
+#include "Serialization/BitWriter.h"
 
 namespace Audio
 {
@@ -15,7 +17,7 @@ namespace Audio
 	 * These parameters are pass to an IAnalyzerNRTFactory when creating
 	 * a new IAnalyzerNRTWorker or IAnalyzerNRTResult
 	 */
-	struct AUDIOANALYZER_API FAnalyzerNRTParameters
+	struct FAnalyzerNRTParameters
 	{
 		public:
 			float SampleRate;
@@ -32,7 +34,7 @@ namespace Audio
 	 * This interface defines the required methods for non-real-time
 	 * analyzer settings.
 	 */
-	class AUDIOANALYZER_API IAnalyzerNRTSettings
+	class IAnalyzerNRTSettings
 	{
 		public:
 			virtual ~IAnalyzerNRTSettings() {};
@@ -43,7 +45,7 @@ namespace Audio
 	 * This interface defines the required methods for non-real-time
 	 * analyzer results.
 	 */
-	class AUDIOANALYZER_API IAnalyzerNRTResult
+	class IAnalyzerNRTResult
 	{
 	public:
 		virtual ~IAnalyzerNRTResult() {};
@@ -53,10 +55,13 @@ namespace Audio
 
 		// This virtual can be overridden to provide a faster copying scheme than full serialization
 		// when analyzing non-remote targets. If not overridden, this function will use the Serialize call.
-		virtual void CopyFrom(IAnalyzerNRTResult* SourceResult);
-
-		// This can be overridden to return a string description of this results struct.
-		virtual FString ToString() const;
+		virtual void CopyFrom(IAnalyzerNRTResult* SourceResult)
+		{
+			FBitWriter SerializedResult;
+			SourceResult->Serialize(SerializedResult);
+			FBitReader DeserializedResult = FBitReader(SerializedResult.GetData(), SerializedResult.GetNumBytes() * 8);
+			Serialize(DeserializedResult);
+		}
 
 		// This must be overridden to return the duration of the original audio analyzed.
 		virtual float GetDurationInSeconds() const = 0;
@@ -67,7 +72,7 @@ namespace Audio
 	 * This interface is used to define a class that will handle actual
 	 * analysis of a singular audio asset.
 	 */
-	class AUDIOANALYZER_API IAnalyzerNRTWorker
+	class IAnalyzerNRTWorker
 	{
 	public:
 		virtual ~IAnalyzerNRTWorker() {};
@@ -88,20 +93,24 @@ namespace Audio
 	 *
 	 * This is used to define a non real-time analyzer.
 	 */
-	class AUDIOANALYZER_API IAnalyzerNRTFactory : public IModularFeature
+	class IAnalyzerNRTFactory : public IModularFeature
 	{
 	public:
 		virtual ~IAnalyzerNRTFactory() {};
 
 		// Supplied unique name of IAnalyzerNRTFactory to enable querying of added 
 		// analyzer factories
-		static FName GetModularFeatureName();
+		static FName GetModularFeatureName()
+		{
+			static FName AudioExtFeatureName = FName(TEXT("AudioAnalyzerNRTPlugin"));
+			return AudioExtFeatureName;
+		}
 
 		// Name of specific analyzer type.
-		virtual FName GetName() const;
+		virtual FName GetName() const = 0;
 
 		// Human readable name of analyzer.
-		virtual FString GetTitle() const;
+		virtual FString GetTitle() const = 0;
 
 		// Create a new result.
 		virtual TUniquePtr<IAnalyzerNRTResult> NewResult() const = 0;
