Updating pitch detector to remove absolution threshold. Remapping sensitivty ranges using power scaling
#rnx

diff --git a/Engine/Plugins/Runtime/AudioSynesthesia/Source/AudioSynesthesiaCore/Private/AutoCorrelationPitchDetector.cpp b/Engine/Plugins/Runtime/AudioSynesthesia/Source/AudioSynesthesiaCore/Private/AutoCorrelationPitchDetector.cpp
--- a/Engine/Plugins/Runtime/AudioSynesthesia/Source/AudioSynesthesiaCore/Private/AutoCorrelationPitchDetector.cpp
+++ b/Engine/Plugins/Runtime/AudioSynesthesia/Source/AudioSynesthesiaCore/Private/AutoCorrelationPitchDetector.cpp
@@ -12,10 +12,36 @@
 
 namespace Audio
 {
+
 	namespace AutoCorrelationPitchDetectorIntrinsics
 	{
 		constexpr int32 MinimumWindowSize = 512;
 		constexpr int32 MaximumWindowSize = 8192;
+		constexpr int32 PeakPickerMinSize = 2;
+
+		// Remap input between 0 and 1 to a space result between the minimum and maximum.
+		// @param InValue - value between 0 and 1 (inclusive)
+		// @param MinimumValue - value greater than or equal to 2.
+		// @param MaximumValue - value greater than or equal to 2.
+		//
+		// @return Remapped value.
+		int32 PowerScalePeakPickerValue(float InValue, int32 MinimumValue, int32 MaximumValue)
+		{
+
+			MinimumValue = FMath::Max(2, MinimumValue);
+			MaximumValue = FMath::Max(2, MaximumValue);
+			InValue = FMath::Clamp(InValue, 0.f, 1.f);
+
+			float Power = FMath::Max(SMALL_NUMBER, FMath::Loge(static_cast<float>(MaximumValue)) / FMath::Loge(static_cast<float>(MinimumValue)));
+			float Denom = FMath::Max(SMALL_NUMBER, Power - 1.f);
+
+			// Power mapped scale between 0 and 1.
+			const float Scale = (FMath::Pow(Power, InValue) - 1.f) / Denom;
+
+			float Range = static_cast<float>(MaximumValue - MinimumValue);
+
+			return FMath::Clamp(MinimumValue + FMath::RoundToInt(Range * Scale), MinimumValue, MaximumValue);
+		}
 	}
 
 	FAutoCorrelationPitchDetector::FAutoCorrelationPitchDetector(const FAutoCorrelationPitchDetectorSettings& InSettings, float InSampleRate)
@@ -66,15 +92,16 @@ namespace Audio
 
 		// setup peak picker
 		float InverseSensitivity = 1.f - Settings.Sensitivity;
-		int32 AutoCorrRange = FMath::Max(1, MaxAutoCorrBin - MinAutoCorrBin);
-		
+		int32 AutoCorrRange = FMath::Max(PeakPickerMinSize, MaxAutoCorrBin - MinAutoCorrBin);
+
+
 		FPeakPickerSettings PeakPickerSettings;
-		PeakPickerSettings.NumPreMax = FMath::Max(1, FMath::RoundToInt(InverseSensitivity * AutoCorrRange / 2.f));
+		PeakPickerSettings.NumPreMax = FMath::Max(1, PowerScalePeakPickerValue(InverseSensitivity, PeakPickerMinSize, AutoCorrRange / 2));
 		PeakPickerSettings.NumPostMax = PeakPickerSettings.NumPreMax;
-		PeakPickerSettings.NumPreMean = FMath::Max(1, FMath::RoundToInt(InverseSensitivity * AutoCorrRange));
+		PeakPickerSettings.NumPreMean = FMath::Max(1, PowerScalePeakPickerValue(InverseSensitivity, PeakPickerMinSize, AutoCorrRange));
 		PeakPickerSettings.NumPostMean = PeakPickerSettings.NumPreMean;
 		PeakPickerSettings.NumWait = 1;
-		PeakPickerSettings.MeanDelta = 0.01f + InverseSensitivity;
+		PeakPickerSettings.MeanDelta = 0.f;
 
 		PeakPicker = MakeUnique<FPeakPicker>(PeakPickerSettings);
 
@@ -123,40 +150,50 @@ namespace Audio
 
 			// Pitch peaks
 			PeakIndices.Reset();
-			PeakPicker->PickPeaks(AutoCorrBuffer, PeakIndices);
+			int32 MaxBin = 0;
 
-			for (int32 Index : PeakIndices)
+			if (AutoCorrBuffer.Num() > MinAutoCorrBin)
 			{
-				if ((Index < MinAutoCorrBin) || (Index > MaxAutoCorrBin))
-				{
-					continue;
-				}
+				float Energy = FMath::Max(SMALL_NUMBER, AutoCorrBuffer[0]);
 
-				float PeakLocOffset = 0.f;
-				float PeakLocMaximum = 0.f;
+				MaxBin = FMath::Min(AutoCorrBuffer.Num() - MinAutoCorrBin - 1, MaxAutoCorrBin - MinAutoCorrBin);
 
-				const float* AutoCorrData = AutoCorrBuffer.GetData();
-				float Energy = FMath::Max(SMALL_NUMBER, AutoCorrData[0]);
+				TArrayView<float> AutoCorrView(&AutoCorrBuffer.GetData()[MinAutoCorrBin], MaxBin + 1);
+				PeakPicker->PickPeaks(AutoCorrView, PeakIndices);
 
-				// Perform quadratic interpolation on peak.
-				// If this returns false, then location is not a true peak.
-				if (QuadraticPeakInterpolation(&AutoCorrData[Index - 1], PeakLocOffset, PeakLocMaximum))
+				for (int32 Index : PeakIndices)
 				{
-					float PeakLoc = static_cast<float>(Index) + PeakLocOffset;
-
-					if (PeakLoc < 2.f)
+					if (Index > MaxBin)
 					{
-						// Invalid peak location.
 						continue;
 					}
 
-					FPitchInfo Info;
+					float PeakLocOffset = 0.f;
+					float PeakLocMaximum = 0.f;
+
+					const float* AutoCorrViewData = AutoCorrView.GetData();
 
-					Info.Frequency = SampleRate / PeakLoc;
-					Info.Strength = FMath::Clamp(PeakLocMaximum / Energy, 0.0f, 1.f);
-					Info.Timestamp = Timestamp;
+					// Perform quadratic interpolation on peak.
+					// If this returns false, then location is not a true peak.
+					if (QuadraticPeakInterpolation(&AutoCorrViewData[Index - 1], PeakLocOffset, PeakLocMaximum))
+					{
+						// Peak adjusted for offset of AutoCorrView and quadratic interpolation.
+						float PeakLoc = static_cast<float>(Index + MinAutoCorrBin) + PeakLocOffset;
+
+						if (PeakLoc < 2.f)
+						{
+							// Invalid peak location.
+							continue;
+						}
+
+						FPitchInfo Info;
 
-					OutPitches.Add(MoveTemp(Info));
+						Info.Frequency = SampleRate / PeakLoc;
+						Info.Strength = FMath::Clamp(PeakLocMaximum / Energy, 0.0f, 1.f);
+						Info.Timestamp = Timestamp;
+
+						OutPitches.Add(MoveTemp(Info));
+					}
 				}
 			}
 		}
diff --git a/Engine/Plugins/Runtime/AudioSynesthesia/Source/AudioSynesthesiaCore/Private/FFTPeakPitchDetector.cpp b/Engine/Plugins/Runtime/AudioSynesthesia/Source/AudioSynesthesiaCore/Private/FFTPeakPitchDetector.cpp
--- a/Engine/Plugins/Runtime/AudioSynesthesia/Source/AudioSynesthesiaCore/Private/FFTPeakPitchDetector.cpp
+++ b/Engine/Plugins/Runtime/AudioSynesthesia/Source/AudioSynesthesiaCore/Private/FFTPeakPitchDetector.cpp
@@ -12,6 +12,28 @@
 
 namespace Audio
 {
+	namespace FFTPeakPitchDetectorIntrinsics
+	{
+		constexpr int32 PeakPickerMinSize = 2;
+
+		int32 PowerScalePeakPickerValue(float InValue, int32 MinimumValue, int32 MaximumValue)
+		{
+			MinimumValue = FMath::Max(2, MinimumValue);
+			MaximumValue = FMath::Max(2, MaximumValue);
+			InValue = FMath::Clamp(InValue, 0.f, 1.f);
+
+			float Power = FMath::Max(SMALL_NUMBER, FMath::Loge(static_cast<float>(MaximumValue)) / FMath::Loge(static_cast<float>(MinimumValue)));
+			float Denom = FMath::Max(SMALL_NUMBER, Power - 1.f);
+
+			// Power mapped scale between 0 and 1.
+			const float Scale = (FMath::Pow(Power, InValue) - 1.f) / Denom;
+
+			float Range = static_cast<float>(MaximumValue - MinimumValue);
+
+			return FMath::Clamp(MinimumValue + FMath::RoundToInt(Range * Scale), MinimumValue, MaximumValue);
+		}
+	}
+
 	FFFTPeakPitchDetector::FFFTPeakPitchDetector(const FFFTPeakPitchDetectorSettings& InSettings, float InSampleRate)
 	:	Settings(InSettings)
 	,	SampleRate(InSampleRate)
@@ -22,6 +44,8 @@ namespace Audio
 	,	WindowCounter(0)
 	,	SlidingBuffer(1, 1)
 	{
+		using namespace FFTPeakPitchDetectorIntrinsics;
+
 		if (!ensure(SampleRate > 0.f))
 		{
 			SampleRate = 1.f;
@@ -83,12 +107,13 @@ namespace Audio
 		float InverseSensitivity = 1.f - Settings.Sensitivity;
 		
 		FPeakPickerSettings PeakPickerSettings;
-		PeakPickerSettings.NumPreMax = FMath::Max(1, FMath::RoundToInt(InverseSensitivity * FFTRange / 2.f));
+		PeakPickerSettings.NumPreMax = FMath::Max(1, PowerScalePeakPickerValue(InverseSensitivity, PeakPickerMinSize, FFTRange / 2));
 		PeakPickerSettings.NumPostMax = PeakPickerSettings.NumPreMax;
-		PeakPickerSettings.NumPreMean = FMath::Max(1, FMath::RoundToInt(InverseSensitivity * FFTRange));
+		PeakPickerSettings.NumPreMean = FMath::Max(1, PowerScalePeakPickerValue(InverseSensitivity, PeakPickerMinSize, FFTRange));
+
 		PeakPickerSettings.NumPostMean = PeakPickerSettings.NumPreMean;
 		PeakPickerSettings.NumWait = 1;
-		PeakPickerSettings.MeanDelta = 0.01f + InverseSensitivity;
+		PeakPickerSettings.MeanDelta = 0.f;
 
 		PeakPicker = MakeUnique<FPeakPicker>(PeakPickerSettings);
 
diff --git a/Engine/Plugins/Runtime/AudioSynesthesia/Source/AudioSynesthesiaCore/Public/AutoCorrelationPitchDetector.h b/Engine/Plugins/Runtime/AudioSynesthesia/Source/AudioSynesthesiaCore/Public/AutoCorrelationPitchDetector.h
--- a/Engine/Plugins/Runtime/AudioSynesthesia/Source/AudioSynesthesiaCore/Public/AutoCorrelationPitchDetector.h
+++ b/Engine/Plugins/Runtime/AudioSynesthesia/Source/AudioSynesthesiaCore/Public/AutoCorrelationPitchDetector.h
@@ -22,8 +22,9 @@ namespace Audio
 		/** Maximum frequency of pitch result. */
 		float MaximumFrequency = 10000.f;
 
-		/** Sensitivity of peak detection */
-		float Sensitivity = 0.9f;
+		/** Sensitivity of peak detection. Valid values are between 0.0 and 1.0. 
+		 * 0.0 results in the least amount of pitches, while 1.0 result in the most pitches. */
+		float Sensitivity = 0.5f;
 	};
 
 	/** Pitch detector based on autocorrelation. Note that autocorrelation pitch detectors 
diff --git a/Engine/Plugins/Runtime/AudioSynesthesia/Source/AudioSynesthesiaCore/Public/FFTPeakPitchDetector.h b/Engine/Plugins/Runtime/AudioSynesthesia/Source/AudioSynesthesiaCore/Public/FFTPeakPitchDetector.h
--- a/Engine/Plugins/Runtime/AudioSynesthesia/Source/AudioSynesthesiaCore/Public/FFTPeakPitchDetector.h
+++ b/Engine/Plugins/Runtime/AudioSynesthesia/Source/AudioSynesthesiaCore/Public/FFTPeakPitchDetector.h
@@ -29,8 +29,9 @@ namespace Audio
 		/** Maximum frequency of detected pitches. */
 		float MaximumFrequency = 10000.f;
 
-		/** Sensitivity used for picking peaks in the power spectrum. */
-		float Sensitivity = 0.9f;
+		/** Sensitivity of peak detection. Valid values are between 0.0 and 1.0. 
+		 * 0.0 results in the least amount of pitches, while 1.0 result in the most pitches. */
+		float Sensitivity = 0.5f;
 
 		/** Reference SNR for calculating strength of signal. */
 		float MaxStrengthSNRDecibels = 20.f;
