Tweaks and cleanup to loudness analyzer

diff --git a/Engine/Plugins/Runtime/AudioSynesthesia/Source/AudioSynesthesia/Classes/Loudness.h b/Engine/Plugins/Runtime/AudioSynesthesia/Source/AudioSynesthesia/Classes/Loudness.h
--- a/Engine/Plugins/Runtime/AudioSynesthesia/Source/AudioSynesthesia/Classes/Loudness.h
+++ b/Engine/Plugins/Runtime/AudioSynesthesia/Source/AudioSynesthesia/Classes/Loudness.h
@@ -35,31 +35,31 @@ class AUDIOSYNESTHESIA_API ULoudnessSettings : public UAudioSynesthesiaSettings
 	GENERATED_BODY()
 public:
 
-	ULoudnessSettings();
+	ULoudnessSettings() {}
 
 	/** Number of seconds between loudness measurements */
 	UPROPERTY(EditAnywhere, BlueprintReadOnly, Category = AudioAnalyzer, meta = (ClampMin = "0.01", ClampMax = "0.25"))
-	float AnalysisPeriod;
+	float AnalysisPeriod = 0.01f;
 
 	/** Minimum analysis frequency for calculating loudness. */
 	UPROPERTY(EditAnywhere, BlueprintReadOnly, Category = AudioAnalyzer, meta = (ClampMin = "20.0", ClampMax = "20000"))
-	float MinimumFrequency;
+	float MinimumFrequency = 20.0f;
 
 	/** Maximum analysis frequency for calculating loudness. */
 	UPROPERTY(EditAnywhere, BlueprintReadOnly, Category = AudioAnalyzer, meta = (ClampMin = "20.0", ClampMax = "20000"))
-	float MaximumFrequency;
+	float MaximumFrequency = 20000.0f;
 
 	/** Type of equal loudness curve used in calculations */
 	UPROPERTY(EditAnywhere, BlueprintReadOnly, Category = AudioAnalyzer)
-	ELoudnessCurveTypeEnum CurveType;
+	ELoudnessCurveTypeEnum CurveType = ELoudnessCurveTypeEnum::D;
 
 	/** dB level to denote silence.  Used when calculating normalized loudness. */
 	UPROPERTY(EditAnywhere, BlueprintReadOnly, AdvancedDisplay, Category = AudioAnalyzer, meta = (ClampMin = "-100.0", ClampMax = "0"))
-	float NoiseFloorDb;
+	float NoiseFloorDb = -60.0f;
 
 	/** dB level to denote silence.  Used when calculating normalized loudness. */
 	UPROPERTY(EditAnywhere, BlueprintReadOnly, AdvancedDisplay, Category = AudioAnalyzer, meta = (ClampMin = "-100.0", ClampMax = "0"))
-	float ExpectedMaxLoudness;
+	float ExpectedMaxLoudness = 0.0f;
 
 	/** Convert ULoudnessSettings to FLoudnessSettings */
 	TUniquePtr<Audio::IAnalyzerSettings> GetSettings(const int32 InSampleRate, const int32 InNumChannels) const;
diff --git a/Engine/Plugins/Runtime/AudioSynesthesia/Source/AudioSynesthesia/Private/Loudness.cpp b/Engine/Plugins/Runtime/AudioSynesthesia/Source/AudioSynesthesia/Private/Loudness.cpp
--- a/Engine/Plugins/Runtime/AudioSynesthesia/Source/AudioSynesthesia/Private/Loudness.cpp
+++ b/Engine/Plugins/Runtime/AudioSynesthesia/Source/AudioSynesthesia/Private/Loudness.cpp
@@ -5,15 +5,6 @@
 #include "InterpolateSorted.h"
 #include "AudioSynesthesiaLog.h"
 
-ULoudnessSettings::ULoudnessSettings()
-	: AnalysisPeriod(0.01f)
-	, MinimumFrequency(20.f)
-	, MaximumFrequency(20000.f)
-	, CurveType(ELoudnessCurveTypeEnum::D)
-	, NoiseFloorDb(-60.f)
-	, ExpectedMaxLoudness(0.0f)
-{}
-
 TUniquePtr<Audio::IAnalyzerSettings> ULoudnessSettings::GetSettings(const int32 InSampleRate, const int32 InNumChannels) const
 {
 	TUniquePtr<Audio::FLoudnessSettings> Settings = MakeUnique<Audio::FLoudnessSettings>();
diff --git a/Engine/Plugins/Runtime/AudioSynesthesia/Source/AudioSynesthesiaCore/Private/LoudnessFactory.cpp b/Engine/Plugins/Runtime/AudioSynesthesia/Source/AudioSynesthesiaCore/Private/LoudnessFactory.cpp
--- a/Engine/Plugins/Runtime/AudioSynesthesia/Source/AudioSynesthesiaCore/Private/LoudnessFactory.cpp
+++ b/Engine/Plugins/Runtime/AudioSynesthesia/Source/AudioSynesthesiaCore/Private/LoudnessFactory.cpp
@@ -9,20 +9,11 @@ namespace Audio
 	// Use for overall loudness storage
 	const int32 FLoudnessResult::ChannelIndexOverall = INDEX_NONE;
 
-	FLoudnessResult::FLoudnessResult()
-	:	DurationInSeconds(0.f)
-	{
-	}
-
 	void FLoudnessResult::Add(const FLoudnessEntry& InDatum)
 	{
 		// Store loudness data in appropriate channel
 		TArray<FLoudnessEntry>& LoudnessArray = ChannelLoudnessArrays.FindOrAdd(InDatum.Channel);
 		LoudnessArray.Add(InDatum);
-
-		// Update intervals per channel
-		FFloatInterval& Interval = ChannelLoudnessIntervals.FindOrAdd(InDatum.Channel);
-		Interval.Include(InDatum.Loudness);
 	}
 
 	const TArray<FLoudnessEntry>& FLoudnessResult::GetChannelLoudnessArray(int32 ChannelIdx) const
diff --git a/Engine/Plugins/Runtime/AudioSynesthesia/Source/AudioSynesthesiaCore/Public/LoudnessAnalyzer.h b/Engine/Plugins/Runtime/AudioSynesthesia/Source/AudioSynesthesiaCore/Public/LoudnessAnalyzer.h
--- a/Engine/Plugins/Runtime/AudioSynesthesia/Source/AudioSynesthesiaCore/Public/LoudnessAnalyzer.h
+++ b/Engine/Plugins/Runtime/AudioSynesthesia/Source/AudioSynesthesiaCore/Public/LoudnessAnalyzer.h
@@ -73,13 +73,13 @@ namespace Audio
             FLoudnessAnalyzer(float InSampleRate, const FLoudnessAnalyzerSettings& InSettings);
 
 			/**
-			 * Calculats the perceptual energy for the input samples.  The number of 
+			 * Calculate the perceptual energy for the input samples.  The number of 
 			 * input samples must be equal to the FFTSize supplied in FLoudnessAnalyzerSettings.
 			 */
             float CalculatePerceptualEnergy(TArrayView<const float> InSampleView);
 
 			/**
-			 * Calculats the loudness for the input samples.  The number of  input samples must 
+			 * Calculate the loudness for the input samples.  The number of  input samples must 
 			 * be equal to the FFTSize supplied in FLoudnessAnalyzerSettings.
 			 */
             float CalculateLoudness(TArrayView<const float> InSampleView);
diff --git a/Engine/Plugins/Runtime/AudioSynesthesia/Source/AudioSynesthesiaCore/Public/LoudnessFactory.h b/Engine/Plugins/Runtime/AudioSynesthesia/Source/AudioSynesthesiaCore/Public/LoudnessFactory.h
--- a/Engine/Plugins/Runtime/AudioSynesthesia/Source/AudioSynesthesiaCore/Public/LoudnessFactory.h
+++ b/Engine/Plugins/Runtime/AudioSynesthesia/Source/AudioSynesthesiaCore/Public/LoudnessFactory.h
@@ -41,7 +41,7 @@ namespace Audio
 		 */
 		static const int32 ChannelIndexOverall;
 
-		FLoudnessResult();
+		FLoudnessResult() {}
 
 		/** Appends an FLoudnessDatum to the container. */
 		void Add(const FLoudnessEntry& InDatum);
@@ -56,9 +56,7 @@ namespace Audio
 		int32 GetNumChannels() const;
 
 	private:
-		float DurationInSeconds;
 		TMap<int32, TArray<FLoudnessEntry> > ChannelLoudnessArrays;
-		TMap<int32, FFloatInterval> ChannelLoudnessIntervals;
 	};
 
 	/**
