Fixing typo on modular feature unregister

diff --git a/Engine/Plugins/Runtime/AudioSynesthesia/Source/AudioSynesthesiaCore/Private/AudioSynesthesiaCoreModule.cpp b/Engine/Plugins/Runtime/AudioSynesthesia/Source/AudioSynesthesiaCore/Private/AudioSynesthesiaCoreModule.cpp
--- a/Engine/Plugins/Runtime/AudioSynesthesia/Source/AudioSynesthesiaCore/Private/AudioSynesthesiaCoreModule.cpp
+++ b/Engine/Plugins/Runtime/AudioSynesthesia/Source/AudioSynesthesiaCore/Private/AudioSynesthesiaCoreModule.cpp
@@ -35,7 +35,7 @@ namespace Audio
 				IModularFeatures::Get().UnregisterModularFeature(FConstantQNRTFactory::GetModularFeatureName(), &ConstantQNRTFactory);
 				IModularFeatures::Get().UnregisterModularFeature(FOnsetNRTFactory::GetModularFeatureName(), &OnsetNRTFactory);
 
-				IModularFeatures::Get().UnregisterModularFeature(FLoudnessNRTFactory::GetModularFeatureName(), &LoudnessFactory);
+				IModularFeatures::Get().UnregisterModularFeature(FLoudnessFactory::GetModularFeatureName(), &LoudnessFactory);
 			}
 
 		private:
