New Meter Analyzer (Sub for aaron.mcleran)

diff --git a/Engine/Plugins/Runtime/AudioSynesthesia/Source/AudioSynesthesia/Classes/Meter.h b/Engine/Plugins/Runtime/AudioSynesthesia/Source/AudioSynesthesia/Classes/Meter.h
new file mode 100644
--- /dev/null
+++ b/Engine/Plugins/Runtime/AudioSynesthesia/Source/AudioSynesthesia/Classes/Meter.h
@@ -0,0 +1,154 @@
+// Copyright Epic Games, Inc. All Rights Reserved.
+
+#pragma once
+
+#include "AudioSynesthesia.h"
+#include "Meter.generated.h"
+
+UENUM(BlueprintType)
+enum class EMeterPeakType : uint8
+{
+	MeanSquared,
+	RootMeanSquared,
+	Peak,
+	Count UMETA(Hidden)
+};
+
+/** UMeterSettings
+ *
+ * Settings for a UMeterAnalyzer.
+ */
+UCLASS(Blueprintable)
+class AUDIOSYNESTHESIA_API UMeterSettings : public UAudioSynesthesiaSettings
+{
+	GENERATED_BODY()
+public:
+
+	UMeterSettings() {}
+
+	/** Number of seconds between meter measurements */
+	UPROPERTY(EditAnywhere, BlueprintReadOnly, Category = AudioAnalyzer, meta = (ClampMin = "0.01", ClampMax = "0.25"))
+	float AnalysisPeriod = 0.01f;
+
+	/** Meter envelope type type */
+	UPROPERTY(EditAnywhere, BlueprintReadOnly, Category = AudioAnalyzer)
+	EMeterPeakType PeakMode = EMeterPeakType::RootMeanSquared;
+
+	/** Meter attack time, in milliseconds */
+	UPROPERTY(EditAnywhere, BlueprintReadOnly, Category = AudioAnalyzer, meta = (ClampMin = "0"))
+	int32 MeterAttackTime = 300;
+
+	/** Meter release time, in milliseconds. */
+	UPROPERTY(EditAnywhere, BlueprintReadOnly, Category = AudioAnalyzer, meta = (ClampMin = "0"))
+	int32 MeterReleaseTime = 300;
+
+	/** Peak attack time, in milliseconds */
+	UPROPERTY(EditAnywhere, BlueprintReadOnly, Category = AudioAnalyzer, meta = (ClampMin = "0"))
+	int32 PeakAttackTime = 300;
+
+	/** Peak release time, in milliseconds */
+	UPROPERTY(EditAnywhere, BlueprintReadOnly, Category = AudioAnalyzer, meta = (ClampMin = "0"))
+	int32 PeakReleaseTime = 300;
+
+	/** Peak hold time, in milliseconds */
+	UPROPERTY(EditAnywhere, BlueprintReadOnly, Category = AudioAnalyzer, meta = (ClampMin = "0"))
+	int32 PeakHoldTime = 100;
+
+	/** What volume threshold to throw clipping detection notifications. */
+	UPROPERTY(EditAnywhere, BlueprintReadOnly, Category = AudioAnalyzer, meta = (ClampMin = "0.0"))
+	float ClippingThreshold = 1.0f;
+
+	/** Convert UMeterSettings to IAnalyzerSettings */
+	TUniquePtr<Audio::IAnalyzerSettings> GetSettings(const int32 InSampleRate, const int32 InNumChannels) const;
+
+#if WITH_EDITOR
+	virtual FText GetAssetActionName() const override;
+
+	virtual UClass* GetSupportedClass() const override;
+#endif
+};
+
+/** The results of the meter analysis. */
+USTRUCT(BlueprintType)
+struct AUDIOSYNESTHESIA_API FMeterResults
+{
+	GENERATED_USTRUCT_BODY()
+
+	// The time in seconds since analysis began of this meter analysis result
+	UPROPERTY(EditAnywhere, BlueprintReadOnly, Category = AudioAnalyzer)
+	float TimeSeconds = 0.0f;
+
+	// The meter value
+	UPROPERTY(EditAnywhere, BlueprintReadOnly, Category = AudioAnalyzer)
+	float MeterValue = 0.0f;
+
+	// The peak value
+	UPROPERTY(EditAnywhere, BlueprintReadOnly, Category = AudioAnalyzer)
+	float PeakValue = 0.0f;
+
+	// The number of samples in the period which were above the clipping threshold. Will be 0 if no clipping was detected.
+	UPROPERTY(EditAnywhere, BlueprintReadOnly, Category = AudioAnalyzer)
+	int32 NumSamplesClipping = 0;
+
+	// The value (if non-zero) if clipping was detected above the clipping threshold
+	UPROPERTY(EditAnywhere, BlueprintReadOnly, Category = AudioAnalyzer)
+	float ClippingValue = 0.0f;
+};
+
+/** Delegate to receive all overall loudness results (time-stamped in an array) since last delegate call. */
+DECLARE_DYNAMIC_MULTICAST_DELEGATE_OneParam(FOnOverallMeterResults, const TArray<FMeterResults>&, MeterResults);
+
+/** Delegate to receive only the most recent overall meter results. */
+DECLARE_DYNAMIC_MULTICAST_DELEGATE_OneParam(FOnLatestOverallMeterResults, const FMeterResults&, LatestOverallMeterResults);
+
+/** Delegate to receive all meter results per channel (time-stamped in an array) since last delegate call. */
+DECLARE_DYNAMIC_MULTICAST_DELEGATE_TwoParams(FOnPerChannelMeterResults, int32, ChannelIndex, const TArray<FMeterResults>&, MeterResults);
+
+/** Delegate to receive only the most recent overall meter result per channel. */
+DECLARE_DYNAMIC_MULTICAST_DELEGATE_TwoParams(FOnLatestPerChannelMeterResults, int32, ChannelIndex, const FMeterResults&, LatestMeterResults);
+
+/** UMeterAnalyzer
+ *
+ * UMeterAnalyzer calculates the current amplitude of an
+ * audio bus in real-time.
+ */
+UCLASS(Blueprintable)
+class AUDIOSYNESTHESIA_API UMeterAnalyzer : public UAudioAnalyzer
+{
+	GENERATED_BODY()
+
+public:
+
+	UMeterAnalyzer();
+
+	/** The settings for the meter audio analyzer.  */
+	UPROPERTY(EditAnywhere, BlueprintReadWrite, Category = AudioAnalyzer)
+	UMeterSettings* Settings;
+
+	/** Delegate to receive all overall meter results since last delegate call. */
+	UPROPERTY(BlueprintAssignable)
+	FOnOverallMeterResults OnOverallMeterResults;
+
+	/** Delegate to receive all meter results, per-channel, since last delegate call. */
+	UPROPERTY(BlueprintAssignable)
+	FOnPerChannelMeterResults OnPerChannelMeterResults;
+
+	/** Delegate to receive the latest overall meter results. */
+	UPROPERTY(BlueprintAssignable)
+	FOnLatestOverallMeterResults OnLatestOverallMeterResults;
+
+	/** Delegate to receive the latest per-channel meter results. */
+	UPROPERTY(BlueprintAssignable)
+	FOnLatestPerChannelMeterResults OnLatestPerChannelMeterResults;
+
+	/** Convert UMeterSettings to IAnalyzerSettings */
+	TUniquePtr<Audio::IAnalyzerSettings> GetSettings(const int32 InSampleRate, const int32 InNumChannels) const override;
+
+	/** Broadcasts results to any delegates if hooked up. */
+	void BroadcastResults() override;
+
+protected:
+
+	/** Return the name of the IAudioAnalyzerFactory associated with this UAudioAnalyzer */
+	FName GetAnalyzerFactoryName() const override;
+};
\ No newline at end of file
diff --git a/Engine/Plugins/Runtime/AudioSynesthesia/Source/AudioSynesthesia/Private/Meter.cpp b/Engine/Plugins/Runtime/AudioSynesthesia/Source/AudioSynesthesia/Private/Meter.cpp
new file mode 100644
--- /dev/null
+++ b/Engine/Plugins/Runtime/AudioSynesthesia/Source/AudioSynesthesia/Private/Meter.cpp
@@ -0,0 +1,162 @@
+// Copyright Epic Games, Inc. All Rights Reserved.
+
+#include "Meter.h"
+#include "MeterFactory.h"
+#include "InterpolateSorted.h"
+#include "AudioSynesthesiaLog.h"
+#include "DSP/EnvelopeFollower.h"
+
+TUniquePtr<Audio::IAnalyzerSettings> UMeterSettings::GetSettings(const int32 InSampleRate, const int32 InNumChannels) const
+{
+	TUniquePtr<Audio::FMeterSettings> Settings = MakeUnique<Audio::FMeterSettings>();
+
+	Settings->AnalysisPeriod = AnalysisPeriod;
+
+	switch (PeakMode)
+	{
+	case EMeterPeakType::MeanSquared:
+		Settings->MeterPeakMode = Audio::EPeakMode::MeanSquared;
+		break;
+
+	case EMeterPeakType::RootMeanSquared:
+		Settings->MeterPeakMode = Audio::EPeakMode::RootMeanSquared;
+		break;
+
+	default:
+	case EMeterPeakType::Peak:
+		Settings->MeterPeakMode = Audio::EPeakMode::Peak;
+		break;
+	}
+
+	Settings->MeterAttackTime = MeterAttackTime;
+	Settings->MeterReleaseTime = MeterReleaseTime;
+	Settings->PeakAttackTime = PeakAttackTime;
+	Settings->PeakReleaseTime = PeakReleaseTime;
+	Settings->PeakHoldTime = PeakHoldTime;
+	Settings->ClippingThreshold = ClippingThreshold;
+
+	return Settings;
+}
+
+#if WITH_EDITOR
+FText UMeterSettings::GetAssetActionName() const
+{
+	return NSLOCTEXT("AssetTypeActions", "AssetTypeActions_AssetSoundSynesthesiaMeterSettings", "Synesthesia Real-Time Settings (Meter)");
+}
+
+UClass* UMeterSettings::GetSupportedClass() const
+{
+	return UMeterSettings::StaticClass();
+}
+#endif
+
+UMeterAnalyzer::UMeterAnalyzer()
+{
+	Settings = CreateDefaultSubobject<UMeterSettings>(TEXT("DefaultMeterSettings"));
+}
+
+TUniquePtr<Audio::IAnalyzerSettings> UMeterAnalyzer::GetSettings(const int32 InSampleRate, const int32 InNumChannels) const
+{
+	TUniquePtr<Audio::IAnalyzerSettings> AnalyzerSettings;
+
+	if (Settings)
+	{
+		AnalyzerSettings = Settings->GetSettings(InSampleRate, InNumChannels);
+	}
+
+	return AnalyzerSettings;
+}
+
+static TArray<FMeterResults> ConvertToBlueprintResults(UMeterSettings* Settings, const TArray<Audio::FMeterEntry>& InMeterArray)
+{
+	TArray<FMeterResults> ResultsArray;
+
+	for (const Audio::FMeterEntry& MeterEntry : InMeterArray)
+	{
+		FMeterResults NewResults;
+
+		NewResults.MeterValue = MeterEntry.MeterValue;
+		NewResults.PeakValue = MeterEntry.PeakValue;
+		NewResults.ClippingValue = MeterEntry.ClippingValue;
+		NewResults.NumSamplesClipping = MeterEntry.NumSamplesClipping;
+		NewResults.TimeSeconds = MeterEntry.Timestamp;
+
+		ResultsArray.Add(NewResults);
+	}
+
+	// Sort by priority (lowest priority first).
+	ResultsArray.Sort([](const FMeterResults& A, const FMeterResults& B) 		
+	{ 
+		return A.TimeSeconds < B.TimeSeconds;
+	});
+
+	return MoveTemp(ResultsArray);
+}
+
+void UMeterAnalyzer::BroadcastResults()
+{
+	TUniquePtr<const Audio::FMeterResult> MeterResults = GetResults<Audio::FMeterResult>();
+	if (!MeterResults.IsValid())
+	{
+		return;
+	}
+
+	int32 NumChannels = MeterResults->GetNumChannels();
+
+	if (NumChannels > 0)
+	{
+		bool bIsOnOverallMeterResultsBound = OnOverallMeterResults.IsBound();
+		bool bIsOnLatestOverallMeterResultsBound = OnLatestOverallMeterResults.IsBound();
+		if (bIsOnOverallMeterResultsBound || bIsOnLatestOverallMeterResultsBound)
+		{
+			const TArray<Audio::FMeterEntry>& OverallMeterArray = MeterResults->GetMeterArray();
+			if (OverallMeterArray.Num() > 0)
+			{
+				TArray<FMeterResults> Results = ConvertToBlueprintResults(Settings, OverallMeterArray);
+				check(Results.Num() > 0);
+
+				if (bIsOnOverallMeterResultsBound)
+				{
+					OnOverallMeterResults.Broadcast(Results);
+				}
+				if (bIsOnLatestOverallMeterResultsBound)
+				{
+					FMeterResults& Latest = Results[Results.Num() - 1];
+					OnLatestOverallMeterResults.Broadcast(Latest);
+				}
+			}
+		}
+
+		bool bIsOnPerChannelMeterResultsBound = OnPerChannelMeterResults.IsBound();
+		bool bIsOnLatestPerChannelMeterResultsBound = OnLatestPerChannelMeterResults.IsBound();
+
+		if (bIsOnPerChannelMeterResultsBound || bIsOnLatestPerChannelMeterResultsBound)
+		{
+			for (int32 ChannelIndex = 0; ChannelIndex < NumChannels; ++ChannelIndex)
+			{
+				const TArray<Audio::FMeterEntry>& MeterArray = MeterResults->GetChannelMeterArray(ChannelIndex);
+				if (MeterArray.Num() > 0)
+				{
+					TArray<FMeterResults> Results = ConvertToBlueprintResults(Settings, MeterArray);
+					check(Results.Num() > 0);
+
+					if (bIsOnPerChannelMeterResultsBound)
+					{
+						OnPerChannelMeterResults.Broadcast(ChannelIndex, Results);
+					}
+
+					if (bIsOnLatestPerChannelMeterResultsBound)
+					{
+						OnLatestPerChannelMeterResults.Broadcast(ChannelIndex, Results[Results.Num() - 1]);
+					}
+				}
+			}
+		}
+	}
+}
+
+FName UMeterAnalyzer::GetAnalyzerFactoryName() const
+{
+	static const FName FactoryName(TEXT("MeterFactory"));
+	return FactoryName;
+}
\ No newline at end of file
diff --git a/Engine/Plugins/Runtime/AudioSynesthesia/Source/AudioSynesthesiaCore/Private/AudioSynesthesiaCoreModule.cpp b/Engine/Plugins/Runtime/AudioSynesthesia/Source/AudioSynesthesiaCore/Private/AudioSynesthesiaCoreModule.cpp
--- a/Engine/Plugins/Runtime/AudioSynesthesia/Source/AudioSynesthesiaCore/Private/AudioSynesthesiaCoreModule.cpp
+++ b/Engine/Plugins/Runtime/AudioSynesthesia/Source/AudioSynesthesiaCore/Private/AudioSynesthesiaCoreModule.cpp
@@ -1,15 +1,16 @@
 // Copyright Epic Games, Inc. All Rights Reserved.
-
 #include "AudioSynesthesiaCoreModule.h"
+
 #include "AudioSynesthesiaCoreLog.h"
+#include "ConstantQNRTFactory.h"
 #include "CoreMinimal.h"
-#include "Modules/ModuleInterface.h"
-#include "Modules/ModuleManager.h"
 #include "Features/IModularFeatures.h"
+#include "LoudnessFactory.h"
 #include "LoudnessNRTFactory.h"
-#include "ConstantQNRTFactory.h"
+#include "MeterFactory.h"
+#include "Modules/ModuleInterface.h"
+#include "Modules/ModuleManager.h"
 #include "OnsetNRTFactory.h"
-#include "LoudnessFactory.h"
 
 DEFINE_LOG_CATEGORY(LogAudioSynesthesiaCore);
 
@@ -26,6 +27,7 @@ namespace Audio
 				IModularFeatures::Get().RegisterModularFeature(FOnsetNRTFactory::GetModularFeatureName(), &OnsetNRTFactory);
 
 				IModularFeatures::Get().RegisterModularFeature(FLoudnessFactory::GetModularFeatureName(), &LoudnessFactory);
+				IModularFeatures::Get().RegisterModularFeature(FMeterFactory::GetModularFeatureName(), &MeterFactory);
 			}
 
 			void ShutdownModule()
@@ -36,6 +38,7 @@ namespace Audio
 				IModularFeatures::Get().UnregisterModularFeature(FOnsetNRTFactory::GetModularFeatureName(), &OnsetNRTFactory);
 
 				IModularFeatures::Get().UnregisterModularFeature(FLoudnessFactory::GetModularFeatureName(), &LoudnessFactory);
+				IModularFeatures::Get().UnregisterModularFeature(FMeterFactory::GetModularFeatureName(), &MeterFactory);
 			}
 
 		private:
@@ -44,6 +47,7 @@ namespace Audio
 			FOnsetNRTFactory OnsetNRTFactory;
 
 			FLoudnessFactory LoudnessFactory;
+			FMeterFactory MeterFactory;
 	};
 
 }
diff --git a/Engine/Plugins/Runtime/AudioSynesthesia/Source/AudioSynesthesiaCore/Private/MeterAnalyzer.cpp b/Engine/Plugins/Runtime/AudioSynesthesia/Source/AudioSynesthesiaCore/Private/MeterAnalyzer.cpp
new file mode 100644
--- /dev/null
+++ b/Engine/Plugins/Runtime/AudioSynesthesia/Source/AudioSynesthesiaCore/Private/MeterAnalyzer.cpp
@@ -0,0 +1,118 @@
+// Copyright Epic Games, Inc. All Rights Reserved.
+
+
+#include "MeterAnalyzer.h"
+#include "AudioMixer.h"
+#include "DSP/BufferVectorOperations.h"
+#include "DSP/DeinterleaveView.h"
+
+
+namespace Audio
+{
+	FMeterAnalyzer::FMeterAnalyzer(float InSampleRate, int32 InNumChannels, const FMeterAnalyzerSettings& InSettings)
+	{
+		NumChannels = InNumChannels;
+		SampleRate = InSampleRate;
+		Settings = InSettings;
+		check(NumChannels > 0);
+		check(SampleRate > 0.0f);
+
+		PeakEnvDataPerChannel.AddDefaulted(NumChannels);
+		PeakDataPerChannel.AddDefaulted(NumChannels);
+		ClippingDataPerChannel.AddDefaulted(NumChannels);
+
+		for (int32 ChannelIndex = 0; ChannelIndex < NumChannels; ++ChannelIndex)
+		{
+			MeterEnvelopeFollowers.Add(MakeUnique<FEnvelopeFollower>(InSampleRate, (float)InSettings.MeterAttackTime, (float)InSettings.MeterReleaseTime, InSettings.MeterPeakMode, true));
+			PeakEnvelopeFollowers.Add(MakeUnique<FEnvelopeFollower>(InSampleRate, (float)InSettings.PeakAttackTime, (float)InSettings.PeakReleaseTime, InSettings.MeterPeakMode, true));
+		}
+	}
+
+	FMeterAnalyzerResults FMeterAnalyzer::ProcessAudio(TArrayView<const float> InSampleView)
+	{
+		// Feed audio through the envelope followers
+		for (int32 SampleIndex = 0; SampleIndex < InSampleView.Num(); SampleIndex += NumChannels)
+		{
+			for (int32 ChannelIndex = 0; ChannelIndex < NumChannels; ++ChannelIndex)
+			{
+				float Sample = InSampleView[SampleIndex + ChannelIndex];
+
+				// Check for clipping in this block of audio
+				FClippingData& ClippingData = ClippingDataPerChannel[ChannelIndex];
+				if (Sample > Settings.ClippingThreshold)
+				{
+					// Track how many samples clipped in this block
+					ClippingData.NumSamplesClipping++;
+
+					// Track the max clipping value in this block of audio
+					ClippingData.ClippingValue = FMath::Max(ClippingData.ClippingValue, Sample);
+				}
+
+				TUniquePtr<FEnvelopeFollower>& MeterEnvelopeFollower = MeterEnvelopeFollowers[ChannelIndex];
+				MeterEnvelopeFollower->ProcessAudio(Sample);
+
+				TUniquePtr<FEnvelopeFollower>& PeakEnvelopeFollower = PeakEnvelopeFollowers[ChannelIndex];
+				float NewPeakEnvValue = PeakEnvelopeFollower->ProcessAudio(Sample);
+
+				FPeakEnvelopeData& PeakEnvData = PeakEnvDataPerChannel[ChannelIndex];
+				FPeakData& PeakData = PeakDataPerChannel[ChannelIndex];
+
+				// Positive if we're rising
+				bool bNewSlopeIsPositive = (NewPeakEnvValue - PeakEnvData.Value) > 0.0f;
+
+				// Detect if the new peak envelope value is less than the previous peak envelope value (i.e. local maximum)
+				if (!bNewSlopeIsPositive && PeakEnvData.bEnvelopeSlopeIsPositive)
+				{
+					// Get the current time
+					float CurrentTimeSec = (PeakEnvData.SampleCount + SampleIndex) / SampleRate;
+
+					// Check if the peak value is larger than the current peak value or if enough time has elapsed to store a new peak
+					if (PeakData.Value < PeakEnvData.Value || (CurrentTimeSec - PeakData.StartTime) >= (0.001f * Settings.PeakHoldTime))
+					{
+						PeakData.Value = PeakEnvData.Value;
+						PeakData.StartTime = CurrentTimeSec;
+					}
+				}
+
+				// Update the current peak envelope data to the new value
+				PeakEnvData.Value = NewPeakEnvValue;
+				PeakEnvData.bEnvelopeSlopeIsPositive = bNewSlopeIsPositive;
+			}
+		}
+
+		// Build the results data
+		FMeterAnalyzerResults Results;
+
+		for (int32 ChannelIndex = 0; ChannelIndex < NumChannels; ++ChannelIndex)
+		{
+			// Add clipping data
+			FClippingData& ClippingData = ClippingDataPerChannel[ChannelIndex];
+
+			Results.ClippingValues.Add(ClippingData.ClippingValue);
+			Results.NumSamplesClipping.Add(ClippingData.NumSamplesClipping);
+
+			// Reset the clipping data for next block of audio
+			ClippingData.ClippingValue = 0.0f;
+			ClippingData.NumSamplesClipping = 0;
+
+			// update sample count of the env data
+			FPeakEnvelopeData& PeakEnvData = PeakEnvDataPerChannel[ChannelIndex];
+			PeakEnvData.SampleCount += InSampleView.Num();
+
+			TUniquePtr<FEnvelopeFollower>& MeterEnvelopeFollower = MeterEnvelopeFollowers[ChannelIndex];
+			Results.MeterValues.Add(MeterEnvelopeFollower->GetCurrentValue());
+
+			FPeakData& PeakData = PeakDataPerChannel[ChannelIndex];
+			Results.PeakValues.Add(PeakData.Value);
+
+			Results.TimeSec = PeakEnvData.SampleCount / SampleRate;
+		}
+
+		return Results;
+	}
+
+	const FMeterAnalyzerSettings& FMeterAnalyzer::GetSettings() const
+	{
+		return Settings;
+	}
+}
diff --git a/Engine/Plugins/Runtime/AudioSynesthesia/Source/AudioSynesthesiaCore/Private/MeterFactory.cpp b/Engine/Plugins/Runtime/AudioSynesthesia/Source/AudioSynesthesiaCore/Private/MeterFactory.cpp
new file mode 100644
--- /dev/null
+++ b/Engine/Plugins/Runtime/AudioSynesthesia/Source/AudioSynesthesiaCore/Private/MeterFactory.cpp
@@ -0,0 +1,146 @@
+// Copyright Epic Games, Inc. All Rights Reserved.
+
+#include "MeterFactory.h"
+#include "AudioSynesthesiaCustomVersion.h"
+#include "DSP/SlidingWindow.h"
+
+namespace Audio
+{
+	// Use for overall meter storage
+	const int32 FMeterResult::ChannelIndexOverall = INDEX_NONE;
+
+	void FMeterResult::Add(const FMeterEntry& InEntry)
+	{
+		// Store meter data in appropriate channel
+		TArray<FMeterEntry>& MeterArray = ChannelMeterArrays.FindOrAdd(InEntry.Channel);
+		MeterArray.Add(InEntry);
+	}
+
+	const TArray<FMeterEntry>& FMeterResult::GetChannelMeterArray(int32 ChannelIdx) const
+	{
+		return ChannelMeterArrays[ChannelIdx];
+	}
+
+	const TArray<FMeterEntry>& FMeterResult::GetMeterArray() const
+	{
+		return ChannelMeterArrays[ChannelIndexOverall];
+	}
+
+	int32 FMeterResult::GetNumChannels() const 
+	{
+		if (ChannelMeterArrays.Num() > 0)
+		{
+			// Channel meter arrays have -1 reserved for overall meter results
+			return ChannelMeterArrays.Num() - 1;
+		}
+		return 0;
+	}
+
+	FMeterWorker::FMeterWorker(const FAnalyzerParameters& InParams, const FMeterSettings& InAnalyzerSettings)
+		: NumChannels(InParams.NumChannels)
+		, NumAnalyzedBuffers(0)
+		, SampleRate(InParams.SampleRate)
+	{
+		check(NumChannels > 0);
+		check(SampleRate > 0);
+
+		NumHopFrames = FMath::CeilToInt(InAnalyzerSettings.AnalysisPeriod * (float)InParams.SampleRate);
+		NumHopFrames = FMath::Max(1, NumHopFrames);
+
+		float NumAnalysisFrames = InAnalyzerSettings.AnalysisPeriod * (float)SampleRate;
+		int32 NumWindowSamples = NumAnalysisFrames * NumChannels;
+		int32 NumHopSamples = NumHopFrames * NumChannels;
+
+		InternalBuffer = MakeUnique<TSlidingBuffer<float>>(NumWindowSamples, NumHopSamples);
+		MeterAnalyzer = MakeUnique<FMeterAnalyzer>(InParams.SampleRate, NumChannels, InAnalyzerSettings);
+	}
+
+	void FMeterWorker::Analyze(TArrayView<const float> InAudio, IAnalyzerResult* OutResult) 
+	{
+		FMeterResult* MeterResult = static_cast<FMeterResult*>(OutResult);
+		check(MeterResult != nullptr);
+
+		TAutoSlidingWindow<float> SlidingWindow(*InternalBuffer, InAudio, InternalWindow);
+
+		// Loop through entire array of input samples. 
+		for (const TArray<float>& Window : SlidingWindow)
+		{
+			const FMeterAnalyzerResults& Results = MeterAnalyzer->ProcessAudio(Window);
+
+			float OverallMeter = 0.0f;
+			float OverallPeak = 0.0f;
+			float Timestamp = 0.0f;
+			int32 MaxClippingSamples = 0;
+			float MaxClippingValue = 0.0f;
+
+			// Add the per-channel results
+			for (int32 ChannelIndex = 0; ChannelIndex < NumChannels; ++ChannelIndex)
+			{
+				FMeterEntry NewEntry = 
+				{ 
+					ChannelIndex, 
+					Results.TimeSec, 
+					Results.MeterValues[ChannelIndex], 
+					Results.PeakValues[ChannelIndex],
+					Results.ClippingValues[ChannelIndex],
+					Results.NumSamplesClipping[ChannelIndex]
+				};
+
+				Timestamp = Results.TimeSec;
+
+				// Add the entries to compute overall
+				OverallMeter += NewEntry.MeterValue;
+				OverallPeak += NewEntry.PeakValue;
+
+				// The "overall" clipping data will be the max number of samples clipped between channels and the max clipping value
+				MaxClippingSamples = FMath::Max(MaxClippingSamples, NewEntry.NumSamplesClipping);
+				MaxClippingValue = FMath::Max(MaxClippingValue, NewEntry.ClippingValue);
+
+				MeterResult->Add(NewEntry);
+			}
+
+			// Divide the overall meter and peak by number of channels to get average
+			OverallMeter /= NumChannels;
+			OverallPeak /= NumChannels;
+
+			// Add to results
+			FMeterEntry OverallEntry = 
+			{ 
+				FMeterResult::ChannelIndexOverall, 
+				Timestamp, 
+				OverallMeter, 
+				OverallPeak, 
+				MaxClippingValue,
+				MaxClippingSamples
+			};
+			MeterResult->Add(OverallEntry);
+		}
+	}
+
+	FName FMeterFactory::GetName() const 
+	{
+		static FName FactoryName(TEXT("MeterFactory"));
+		return FactoryName;
+	}
+
+	FString FMeterFactory::GetTitle() const
+	{
+		return TEXT("Real-Time Meter Analyzer");
+	}
+
+	TUniquePtr<IAnalyzerResult> FMeterFactory::NewResult() const
+	{
+		TUniquePtr<FMeterResult> Result = MakeUnique<FMeterResult>();
+		return Result;
+	}
+
+	TUniquePtr<IAnalyzerWorker> FMeterFactory::NewWorker(const FAnalyzerParameters& InParams, const IAnalyzerSettings* InSettings) const
+	{
+		const FMeterSettings* MeterSettings = static_cast<const FMeterSettings*>(InSettings);
+
+		check(nullptr != MeterSettings);
+
+		return MakeUnique<FMeterWorker>(InParams, *MeterSettings);
+	}
+}
+
diff --git a/Engine/Plugins/Runtime/AudioSynesthesia/Source/AudioSynesthesiaCore/Public/MeterAnalyzer.h b/Engine/Plugins/Runtime/AudioSynesthesia/Source/AudioSynesthesiaCore/Public/MeterAnalyzer.h
new file mode 100644
--- /dev/null
+++ b/Engine/Plugins/Runtime/AudioSynesthesia/Source/AudioSynesthesiaCore/Public/MeterAnalyzer.h
@@ -0,0 +1,102 @@
+// Copyright Epic Games, Inc. All Rights Reserved.
+
+#pragma once
+
+#include "DSP/EnvelopeFollower.h"
+
+namespace Audio
+{
+	struct AUDIOSYNESTHESIACORE_API FMeterAnalyzerSettings
+	{	
+		// Envelope follower mode
+		EPeakMode::Type MeterPeakMode = EPeakMode::RootMeanSquared;
+
+		// Envelope follower attack time in milliseconds
+		int32 MeterAttackTime = 300;
+
+		// Envelope follower release time in milliseconds
+		int32 MeterReleaseTime = 300;
+		
+		// Peak detector attack time in milliseconds
+		int32 PeakAttackTime = 300;
+
+		// Peak detector release time in milliseconds
+		int32 PeakReleaseTime = 300;
+
+		// Peak detector hold time in milliseconds
+		int32 PeakHoldTime = 100;
+
+		// The volume threshold to detect clipping
+		float ClippingThreshold = 1.0f;
+	};
+
+	// Per-channel analyzer results
+	struct FMeterAnalyzerResults
+	{
+		float TimeSec = 0.0f;
+		TArray<float> MeterValues;
+		TArray<float> PeakValues;
+		TArray<float> ClippingValues;
+		TArray<int32> NumSamplesClipping;
+	};
+
+	class FMeterAnalyzer
+	{
+	public:
+
+		/** Construct analyzer */
+		FMeterAnalyzer(float InSampleRate, int32 NumChannels, const FMeterAnalyzerSettings& InSettings);
+
+		/**
+		 * Calculate the meter results for the input samples.  Will return the current meter and the current peak value of the analyzer.
+		 */
+		FMeterAnalyzerResults ProcessAudio(TArrayView<const float> InSampleView);
+
+		/**
+		 * Return const reference to settings used inside this analyzer.
+		 */
+		const FMeterAnalyzerSettings& GetSettings() const;
+
+	protected:
+		// Envelope follower per channel
+		TArray<TUniquePtr<FEnvelopeFollower>> MeterEnvelopeFollowers;
+		TArray<TUniquePtr<FEnvelopeFollower>> PeakEnvelopeFollowers;
+
+		// Per-channel clipping data
+		struct FClippingData
+		{
+			// How many samples we've been clipping
+			int32 NumSamplesClipping = 0;
+			float ClippingValue = 0.0f;
+		};
+		TArray<FClippingData> ClippingDataPerChannel;
+
+		// State to track the peak data
+		struct FPeakData
+		{
+			float Value = 0.0f;
+			float StartTime = 0.0f;
+		};
+		// Per-channel peak data
+		TArray<FPeakData> PeakDataPerChannel;
+
+		struct FPeakEnvelopeData
+		{
+			// Current envelope data
+			float Value = 0.0f;
+
+			// Running sample count of the envelope data
+			int32 SampleCount = 0;
+
+			// Current slope of the envelope
+			bool bEnvelopeSlopeIsPositive = false;
+		};
+		// Per-channel envelope data
+		TArray<FPeakEnvelopeData> PeakEnvDataPerChannel;
+
+		FMeterAnalyzerSettings Settings;
+		float SampleRate = 0.0f;
+		int32 NumChannels = 0;
+	};
+}
+
diff --git a/Engine/Plugins/Runtime/AudioSynesthesia/Source/AudioSynesthesiaCore/Public/MeterFactory.h b/Engine/Plugins/Runtime/AudioSynesthesia/Source/AudioSynesthesiaCore/Public/MeterFactory.h
new file mode 100644
--- /dev/null
+++ b/Engine/Plugins/Runtime/AudioSynesthesia/Source/AudioSynesthesiaCore/Public/MeterFactory.h
@@ -0,0 +1,107 @@
+// Copyright Epic Games, Inc. All Rights Reserved.
+
+#pragma once
+
+#include "MeterAnalyzer.h"
+#include "IAudioAnalyzerInterface.h"
+#include "DSP/SlidingWindow.h"
+
+namespace Audio
+{
+	/**
+	 * Contains settings for meter  analyzer.
+	 */
+	class AUDIOSYNESTHESIACORE_API FMeterSettings : public IAnalyzerSettings, public FMeterAnalyzerSettings
+	{	
+	public:
+		/** Number of seconds between loudness measurements */
+		float AnalysisPeriod = 0.01f;
+	};
+
+	/**
+	 * Holds the meter results per a time step for each channel
+	 */
+	struct FMeterEntry
+	{
+		int32 Channel = 0;
+		float Timestamp = 0.0f;
+		float MeterValue = 0.0f;
+		float PeakValue = 0.0f;
+		float ClippingValue = 0.0f;
+		int32 NumSamplesClipping = 0;
+	};
+
+	/** 
+	 * FMeterResult contains the temporal evolution of loudness for the last buffer of audio analyzed.
+	 */
+	class AUDIOSYNESTHESIACORE_API FMeterResult : public IAnalyzerResult
+	{
+	public:
+		/**
+		 * Denotes the overall loudness as opposed individual channel indices.
+		 */
+		static const int32 ChannelIndexOverall;
+
+		FMeterResult() {}
+
+		/** Appends an FMeterEntry to the container. */
+		void Add(const FMeterEntry& InEntry);
+
+		/** Returns const reference to FMeterEntry array for individual channel. */
+		const TArray<FMeterEntry>& GetChannelMeterArray(int32 ChannelIdx) const;
+
+		/** Returns const reference to FMeterEntry array associated with overall loudness. */
+		const TArray<FMeterEntry>& GetMeterArray() const;
+
+		/** Returns the number of channels. */
+		int32 GetNumChannels() const;
+
+	private:
+		float DurationInSeconds = 0.0f;
+		TMap<int32, TArray<FMeterEntry> > ChannelMeterArrays;
+	};
+
+	/**
+	 * FMeterWorker performs meter analysis on input sample buffers.
+	 */
+	class AUDIOSYNESTHESIACORE_API FMeterWorker : public IAnalyzerWorker
+	{
+	public:
+		/** Construct a worker */
+		FMeterWorker(const FAnalyzerParameters& InParams, const FMeterSettings& InAnalyzerSettings);
+
+		/** Analyzes input sample buffer and updates result. */
+		virtual void Analyze(TArrayView<const float> InAudio, IAnalyzerResult* OutResult) override;
+
+	private:
+
+		int32 NumChannels = 0;
+		int32 NumAnalyzedBuffers = 0;
+		int32 NumHopFrames = 0;
+		int32 SampleRate = 0;
+		TArray<float> InternalWindow;
+		TUniquePtr<TSlidingBuffer<float>> InternalBuffer;
+		TUniquePtr<FMeterAnalyzer> MeterAnalyzer;
+	};
+
+	/**
+	 * Defines the meter analyzer and creates related classes.
+	 */
+	class AUDIOSYNESTHESIACORE_API FMeterFactory : public IAnalyzerFactory
+	{
+		public:
+
+		/** Name of specific analyzer type. */
+		virtual FName GetName() const override;
+
+		/** Human readable name of analyzer. */
+		virtual FString GetTitle() const override;
+
+		/** Creates a new FLoudnessNRTResult */
+		virtual TUniquePtr<IAnalyzerResult> NewResult() const override;
+
+		/**  Creates a new FLoudnessWorker. This expects IAnalyzerNRTSettings to be a valid pointer to a FLoudnessSettings object. */
+		virtual TUniquePtr<IAnalyzerWorker> NewWorker(const FAnalyzerParameters& InParams, const IAnalyzerSettings* InSettings) const override;
+	};
+}
+
