Manually integrating 15343832

diff --git a/Engine/Plugins/Runtime/AudioSynesthesia/Source/AudioSynesthesia/Private/Meter.cpp b/Engine/Plugins/Runtime/AudioSynesthesia/Source/AudioSynesthesia/Private/Meter.cpp
--- a/Engine/Plugins/Runtime/AudioSynesthesia/Source/AudioSynesthesia/Private/Meter.cpp
+++ b/Engine/Plugins/Runtime/AudioSynesthesia/Source/AudioSynesthesia/Private/Meter.cpp
@@ -105,7 +105,7 @@ void UMeterAnalyzer::BroadcastResults()
 	if (NumChannels > 0)
 	{
 		bool bIsOnOverallMeterResultsBound = OnOverallMeterResults.IsBound() || OnOverallMeterResultsNative.IsBound();
-		bool bIsOnLatestOverallMeterResultsBound = OnLatestOverallMeterResults.IsBound() || OnLatestOverallMeterResults.IsBound();
+		bool bIsOnLatestOverallMeterResultsBound = OnLatestOverallMeterResults.IsBound() || OnLatestOverallMeterResultsNative.IsBound();
 		if (bIsOnOverallMeterResultsBound || bIsOnLatestOverallMeterResultsBound)
 		{
 			const TArray<Audio::FMeterEntry>& OverallMeterArray = MeterResults->GetMeterArray();
