Audio Plugin Telemetry Part 1
Audio Capture, Audio Modulation, Audio Synesthesia

diff --git a/Engine/Plugins/Runtime/AudioCapture/Source/AudioCapture/Private/AudioCaptureComponent.cpp b/Engine/Plugins/Runtime/AudioCapture/Source/AudioCapture/Private/AudioCaptureComponent.cpp
--- a/Engine/Plugins/Runtime/AudioCapture/Source/AudioCapture/Private/AudioCaptureComponent.cpp
+++ b/Engine/Plugins/Runtime/AudioCapture/Source/AudioCapture/Private/AudioCaptureComponent.cpp
@@ -1,6 +1,7 @@
 // Copyright Epic Games, Inc. All Rights Reserved.
 
 #include "AudioCaptureComponent.h"
+#include "EngineAnalytics.h"
 
 static const unsigned int MaxBufferSize = 2 * 5 * 48000;
 
@@ -29,6 +30,12 @@ bool UAudioCaptureComponent::Init(int32& SampleRate)
 		{
 			// This may fail if capture synths aren't supported on a given platform or if something went wrong with the capture device
 			bIsStreamOpen = CaptureSynth.OpenDefaultStream();
+
+			if (FEngineAnalytics::IsAvailable())
+			{
+				FEngineAnalytics::GetProvider().RecordEvent(TEXT("Audio.Usage.AudioCapture.AudioCaptureComponentInitialized"));
+			}
+
 			return true;
 		}
 		else
diff --git a/Engine/Plugins/Runtime/AudioModulation/Source/AudioModulationEditor/Private/Factories/SoundControlBusFactory.cpp b/Engine/Plugins/Runtime/AudioModulation/Source/AudioModulationEditor/Private/Factories/SoundControlBusFactory.cpp
--- a/Engine/Plugins/Runtime/AudioModulation/Source/AudioModulationEditor/Private/Factories/SoundControlBusFactory.cpp
+++ b/Engine/Plugins/Runtime/AudioModulation/Source/AudioModulationEditor/Private/Factories/SoundControlBusFactory.cpp
@@ -1,6 +1,8 @@
 // Copyright Epic Games, Inc. All Rights Reserved.
 
 #include "SoundControlBusFactory.h"
+
+#include "EngineAnalytics.h"
 #include "SoundControlBus.h"
 
 
@@ -15,5 +17,9 @@ USoundControlBusFactory::USoundControlBusFactory(const FObjectInitializer& Objec
 
 UObject* USoundControlBusFactory::FactoryCreateNew(UClass* Class, UObject* InParent, FName Name, EObjectFlags Flags, UObject* Context, FFeedbackContext* Warn)
 {
+	if (FEngineAnalytics::IsAvailable())
+	{
+		FEngineAnalytics::GetProvider().RecordEvent(TEXT("Audio.Usage.AudioModulation.ControlBusCreated"));
+	}
 	return NewObject<USoundControlBus>(InParent, Name, Flags);
 }
\ No newline at end of file
diff --git a/Engine/Plugins/Runtime/AudioModulation/Source/AudioModulationEditor/Private/Factories/SoundControlBusMixFactory.cpp b/Engine/Plugins/Runtime/AudioModulation/Source/AudioModulationEditor/Private/Factories/SoundControlBusMixFactory.cpp
--- a/Engine/Plugins/Runtime/AudioModulation/Source/AudioModulationEditor/Private/Factories/SoundControlBusMixFactory.cpp
+++ b/Engine/Plugins/Runtime/AudioModulation/Source/AudioModulationEditor/Private/Factories/SoundControlBusMixFactory.cpp
@@ -1,6 +1,8 @@
 // Copyright Epic Games, Inc. All Rights Reserved.
 
 #include "SoundControlBusMixFactory.h"
+
+#include "EngineAnalytics.h"
 #include "SoundControlBusMix.h"
 
 
@@ -15,5 +17,9 @@ USoundControlBusMixFactory::USoundControlBusMixFactory(const FObjectInitializer&
 
 UObject* USoundControlBusMixFactory::FactoryCreateNew(UClass* Class, UObject* InParent, FName Name, EObjectFlags Flags, UObject* Context, FFeedbackContext* Warn)
 {
+	if (FEngineAnalytics::IsAvailable())
+	{
+		FEngineAnalytics::GetProvider().RecordEvent(TEXT("Audio.Usage.AudioModulation.ControlBusMixCreated"));
+	}
 	return NewObject<USoundControlBusMix>(InParent, Name, Flags);
 }
diff --git a/Engine/Plugins/Runtime/AudioModulation/Source/AudioModulationEditor/Private/Factories/SoundModulationGeneratorFactory.cpp b/Engine/Plugins/Runtime/AudioModulation/Source/AudioModulationEditor/Private/Factories/SoundModulationGeneratorFactory.cpp
--- a/Engine/Plugins/Runtime/AudioModulation/Source/AudioModulationEditor/Private/Factories/SoundModulationGeneratorFactory.cpp
+++ b/Engine/Plugins/Runtime/AudioModulation/Source/AudioModulationEditor/Private/Factories/SoundModulationGeneratorFactory.cpp
@@ -4,6 +4,7 @@
 
 #include "ClassViewerFilter.h"
 #include "ClassViewerModule.h"
+#include "EngineAnalytics.h"
 #include "Kismet2/SClassPickerDialog.h"
 #include "SoundModulationGenerator.h"
 #include "Templates/SharedPointer.h"
@@ -75,6 +76,10 @@ UObject* USoundModulationGeneratorFactory::FactoryCreateNew(UClass* InClass, UOb
 {
 	if (GeneratorClass)
 	{
+		if (FEngineAnalytics::IsAvailable())
+		{
+			FEngineAnalytics::GetProvider().RecordEvent(TEXT("Audio.Usage.AudioModulation.ModulationGeneratorCreated"));
+		}
 		return NewObject<USoundModulationGenerator>(InParent, GeneratorClass, InName, Flags);
 	}
 
diff --git a/Engine/Plugins/Runtime/AudioModulation/Source/AudioModulationEditor/Private/Factories/SoundModulationParameterFactory.cpp b/Engine/Plugins/Runtime/AudioModulation/Source/AudioModulationEditor/Private/Factories/SoundModulationParameterFactory.cpp
--- a/Engine/Plugins/Runtime/AudioModulation/Source/AudioModulationEditor/Private/Factories/SoundModulationParameterFactory.cpp
+++ b/Engine/Plugins/Runtime/AudioModulation/Source/AudioModulationEditor/Private/Factories/SoundModulationParameterFactory.cpp
@@ -3,6 +3,7 @@
 #include "SoundModulationParameterFactory.h"
 
 #include "ClassViewerFilter.h"
+#include "EngineAnalytics.h"
 #include "SoundModulationParameter.h"
 #include "Kismet2/SClassPickerDialog.h"
 
@@ -71,6 +72,12 @@ UObject* USoundModulationParameterFactory::FactoryCreateNew(UClass* InClass, UOb
 	if (ParameterClass != nullptr)
 	{
 		NewParameter = NewObject<USoundModulationParameter>(InParent, ParameterClass, InName, Flags);
+
+		if (FEngineAnalytics::IsAvailable())
+		{
+			FEngineAnalytics::GetProvider().RecordEvent(TEXT("Audio.Usage.AudioModulation.ModulationParameterCreated"));
+		}
 	}
+
 	return NewParameter;
 }
\ No newline at end of file
diff --git a/Engine/Plugins/Runtime/AudioModulation/Source/AudioModulationEditor/Private/Factories/SoundModulationPatchFactory.cpp b/Engine/Plugins/Runtime/AudioModulation/Source/AudioModulationEditor/Private/Factories/SoundModulationPatchFactory.cpp
--- a/Engine/Plugins/Runtime/AudioModulation/Source/AudioModulationEditor/Private/Factories/SoundModulationPatchFactory.cpp
+++ b/Engine/Plugins/Runtime/AudioModulation/Source/AudioModulationEditor/Private/Factories/SoundModulationPatchFactory.cpp
@@ -1,5 +1,7 @@
 // Copyright Epic Games, Inc. All Rights Reserved.
 #include "SoundModulationPatchFactory.h"
+
+#include "EngineAnalytics.h"
 #include "SoundModulationPatch.h"
 
 
@@ -14,5 +16,9 @@ USoundModulationPatchFactory::USoundModulationPatchFactory(const FObjectInitiali
 
 UObject* USoundModulationPatchFactory::FactoryCreateNew(UClass* Class, UObject* InParent, FName Name, EObjectFlags Flags, UObject* Context, FFeedbackContext* Warn)
 {
+	if (FEngineAnalytics::IsAvailable())
+	{
+		FEngineAnalytics::GetProvider().RecordEvent(TEXT("Audio.Usage.AudioModulation.ParameterPatchCreated"));
+	}
 	return NewObject<USoundModulationPatch>(InParent, Name, Flags);
 }
diff --git a/Engine/Plugins/Runtime/AudioSynesthesia/Source/AudioSynesthesiaEditor/Private/AudioSynesthesiaNRTFactory.cpp b/Engine/Plugins/Runtime/AudioSynesthesia/Source/AudioSynesthesiaEditor/Private/AudioSynesthesiaNRTFactory.cpp
--- a/Engine/Plugins/Runtime/AudioSynesthesia/Source/AudioSynesthesiaEditor/Private/AudioSynesthesiaNRTFactory.cpp
+++ b/Engine/Plugins/Runtime/AudioSynesthesia/Source/AudioSynesthesiaEditor/Private/AudioSynesthesiaNRTFactory.cpp
@@ -4,6 +4,7 @@
 
 #include "ClassViewerModule.h"
 #include "ClassViewerFilter.h"
+#include "EngineAnalytics.h"
 #include "Kismet2/SClassPickerDialog.h"
 #include "Modules/ModuleManager.h"
 #include "AudioSynesthesiaClassFilter.h"
@@ -55,6 +56,11 @@ UObject* UAudioSynesthesiaNRTFactory::FactoryCreateNew(UClass* InClass, UObject*
 	if (AudioSynesthesiaNRTClass != nullptr)
 	{
 		NewAudioSynesthesiaNRT = NewObject<UAudioSynesthesiaNRT>(InParent, AudioSynesthesiaNRTClass, InName, Flags);
+
+		if (FEngineAnalytics::IsAvailable())
+		{
+			FEngineAnalytics::GetProvider().RecordEvent(TEXT("Audio.Usage.AudioSynesthesia.NRTCreated"));
+		}
 	}
 	return NewAudioSynesthesiaNRT;
 }
diff --git a/Engine/Plugins/Runtime/AudioSynesthesia/Source/AudioSynesthesiaEditor/Private/AudioSynesthesiaNRTSettingsFactory.cpp b/Engine/Plugins/Runtime/AudioSynesthesia/Source/AudioSynesthesiaEditor/Private/AudioSynesthesiaNRTSettingsFactory.cpp
--- a/Engine/Plugins/Runtime/AudioSynesthesia/Source/AudioSynesthesiaEditor/Private/AudioSynesthesiaNRTSettingsFactory.cpp
+++ b/Engine/Plugins/Runtime/AudioSynesthesia/Source/AudioSynesthesiaEditor/Private/AudioSynesthesiaNRTSettingsFactory.cpp
@@ -4,6 +4,7 @@
 
 #include "ClassViewerModule.h"
 #include "ClassViewerFilter.h"
+#include "EngineAnalytics.h"
 #include "Kismet2/SClassPickerDialog.h"
 #include "Modules/ModuleManager.h"
 #include "AudioSynesthesiaClassFilter.h"
@@ -54,6 +55,10 @@ UObject* UAudioSynesthesiaNRTSettingsFactory::FactoryCreateNew(UClass* InClass,
 	if (AudioSynesthesiaNRTSettingsClass != nullptr)
 	{
 		NewAudioSynesthesiaNRTSettings = NewObject<UAudioSynesthesiaNRTSettings>(InParent, AudioSynesthesiaNRTSettingsClass, InName, Flags);
+		if (FEngineAnalytics::IsAvailable())
+		{
+			FEngineAnalytics::GetProvider().RecordEvent(TEXT("Audio.Usage.AudioSynesthesia.NRTSettingsCreated"));
+		}
 	}
 	return NewAudioSynesthesiaNRTSettings;
 }
diff --git a/Engine/Plugins/Runtime/AudioSynesthesia/Source/AudioSynesthesiaEditor/Private/AudioSynesthesiaSettingsFactory.cpp b/Engine/Plugins/Runtime/AudioSynesthesia/Source/AudioSynesthesiaEditor/Private/AudioSynesthesiaSettingsFactory.cpp
--- a/Engine/Plugins/Runtime/AudioSynesthesia/Source/AudioSynesthesiaEditor/Private/AudioSynesthesiaSettingsFactory.cpp
+++ b/Engine/Plugins/Runtime/AudioSynesthesia/Source/AudioSynesthesiaEditor/Private/AudioSynesthesiaSettingsFactory.cpp
@@ -4,6 +4,7 @@
 
 #include "ClassViewerModule.h"
 #include "ClassViewerFilter.h"
+#include "EngineAnalytics.h"
 #include "Kismet2/SClassPickerDialog.h"
 #include "Modules/ModuleManager.h"
 #include "AudioSynesthesiaClassFilter.h"
@@ -54,6 +55,11 @@ UObject* UAudioSynesthesiaSettingsFactory::FactoryCreateNew(UClass* InClass, UOb
 	if (AudioSynesthesiaSettingsClass != nullptr)
 	{
 		NewAudioSynesthesiaSettings = NewObject<UAudioSynesthesiaSettings>(InParent, AudioSynesthesiaSettingsClass, InName, Flags);
+
+		if (FEngineAnalytics::IsAvailable())
+		{
+			FEngineAnalytics::GetProvider().RecordEvent(TEXT("Audio.Usage.AudioSynesthesia.SettingsFactoryCreated"));
+		}
 	}
 	return NewAudioSynesthesiaSettings;
 }
