Audio Synesthesia real time FFT analyzer (audio bus spectral analysis)

diff --git a/Engine/Plugins/Runtime/AudioSynesthesia/Source/AudioSynesthesia/Classes/SynesthesiaSpectrumAnalysis.h b/Engine/Plugins/Runtime/AudioSynesthesia/Source/AudioSynesthesia/Classes/SynesthesiaSpectrumAnalysis.h
new file mode 100644
--- /dev/null
+++ b/Engine/Plugins/Runtime/AudioSynesthesia/Source/AudioSynesthesia/Classes/SynesthesiaSpectrumAnalysis.h
@@ -0,0 +1,125 @@
+// Copyright Epic Games, Inc. All Rights Reserved.
+
+#pragma once
+
+#include "AudioSynesthesia.h"
+#include "Sound/SoundSubmix.h"
+#include "SynesthesiaSpectrumAnalysisFactory.h"
+#include "SynesthesiaSpectrumAnalysis.generated.h"
+
+/** USynesthesiaSpectrumAnalysisSettings
+ *
+ * Settings for a USynesthesiaSpectrumAnalysisAnalyzer.
+ */
+UCLASS(Blueprintable)
+class AUDIOSYNESTHESIA_API USynesthesiaSpectrumAnalysisSettings : public UAudioSynesthesiaSettings
+{
+	GENERATED_BODY()
+public:
+
+	USynesthesiaSpectrumAnalysisSettings() {}
+
+	/** Number of seconds between SynesthesiaSpectrumAnalysis measurements */
+	UPROPERTY(EditAnywhere, BlueprintReadOnly, Category = AudioAnalyzer, meta = (ClampMin = "0.01", ClampMax = "0.25"))
+	float AnalysisPeriod = 0.01f;
+
+	/** Size of FFT. */
+	UPROPERTY(EditAnywhere, BlueprintReadOnly, Category = AudioAnalyzer)
+	EFFTSize FFTSize = EFFTSize::DefaultSize;
+	
+	/** Type of spectrum to use. */
+	UPROPERTY(EditAnywhere, BlueprintReadOnly, Category = AudioAnalyzer)
+	EAudioSpectrumType SpectrumType = EAudioSpectrumType::PowerSpectrum;
+
+	/** Type of window to use. */
+	UPROPERTY(EditAnywhere, BlueprintReadOnly, Category = AudioAnalyzer)
+	EFFTWindowType WindowType = EFFTWindowType::Hann;
+
+	/** If true, multichannel audio is downmixed to mono with equal amplitude scaling. If false, each channel gets its own FFT result. */
+	UPROPERTY(EditAnywhere, BlueprintReadOnly, Category = AudioAnalyzer)
+	bool bDownmixToMono = true;
+
+	/** Convert USynesthesiaSpectrumAnalysisSettings to IAnalyzerSettings */
+	TUniquePtr<Audio::IAnalyzerSettings> GetSettings(const int32 InSampleRate, const int32 InNumChannels) const;
+
+#if WITH_EDITOR
+	virtual FText GetAssetActionName() const override;
+
+	virtual UClass* GetSupportedClass() const override;
+#endif
+};
+
+/** The results of the spectrum analysis. */
+USTRUCT(BlueprintType)
+struct AUDIOSYNESTHESIA_API FSynesthesiaSpectrumResults
+{
+	GENERATED_USTRUCT_BODY()
+
+	// The time in seconds since analysis began of this SynesthesiaSpectrumAnalysis analysis result
+	UPROPERTY(EditAnywhere, BlueprintReadOnly, Category = AudioAnalyzer)
+	float TimeSeconds = 0.0f;
+
+	// The spectrum values from the FFT
+	UPROPERTY(EditAnywhere, BlueprintReadOnly, Category = AudioAnalyzer)
+	TArray<float> SpectrumValues;
+};
+
+/** Delegate to receive all spectrum results per channel (time-stamped in an array) since last delegate call. If bDownmixToMono setting is true, results will be in channel index 0. */
+DECLARE_DYNAMIC_MULTICAST_DELEGATE_TwoParams(FOnSpectrumResults, int32, ChannelIndex, const TArray<FSynesthesiaSpectrumResults>&, SpectrumResults);
+
+/** shadow delegate declaration for above */
+DECLARE_MULTICAST_DELEGATE_ThreeParams(FOnSpectrumResultsNative, USynesthesiaSpectrumAnalyzer*, int32, const TArray<FSynesthesiaSpectrumResults>&);
+
+/** Delegate to receive only the most recent overall spectrum result per channel. If bDownmixToMono setting is true, results will be in channel index 0.*/
+DECLARE_DYNAMIC_MULTICAST_DELEGATE_TwoParams(FOnLatestSpectrumResults, int32, ChannelIndex, const FSynesthesiaSpectrumResults&, LatestSpectrumResults);
+
+/** shadow delegate declaration for above */
+DECLARE_MULTICAST_DELEGATE_ThreeParams(FOnLatestSpectrumResultsNative, USynesthesiaSpectrumAnalyzer*, int32, const FSynesthesiaSpectrumResults&);
+
+/** USynesthesiaSpectrumAnalysisAnalyzer
+ *
+ * USynesthesiaSpectrumAnalysisAnalyzer calculates the current amplitude of an
+ * audio bus in real-time.
+ */
+UCLASS(Blueprintable)
+class AUDIOSYNESTHESIA_API USynesthesiaSpectrumAnalyzer : public UAudioAnalyzer
+{
+	GENERATED_BODY()
+
+public:
+
+	USynesthesiaSpectrumAnalyzer();
+
+	/** The settings for the SynesthesiaSpectrumAnalysis audio analyzer.  */
+	UPROPERTY(EditAnywhere, BlueprintReadWrite, Category = AudioAnalyzer)
+	TObjectPtr<USynesthesiaSpectrumAnalysisSettings> Settings;
+
+	/** Delegate to receive all Spectrum results, per-channel, since last delegate call. If bDownmixToMono setting is true, results will be in channel index 0. */
+	UPROPERTY(BlueprintAssignable)
+	FOnSpectrumResults OnSpectrumResults;
+
+	FOnSpectrumResultsNative OnSpectrumResultsNative;
+
+	/** Delegate to receive the latest per-channel Spectrum results. If bDownmixToMono setting is true, results will be in channel index 0. */
+	UPROPERTY(BlueprintAssignable)
+	FOnLatestSpectrumResults OnLatestSpectrumResults;
+
+	FOnLatestSpectrumResultsNative OnLatestSpectrumResultsNative;
+
+	/** Convert USynesthesiaSpectrumAnalysisSettings to IAnalyzerSettings */
+	TUniquePtr<Audio::IAnalyzerSettings> GetSettings(const int32 InSampleRate, const int32 InNumChannels) const override;
+
+	/** Broadcasts results to any delegates if hooked up. */
+	void BroadcastResults() override;
+
+	UFUNCTION(BlueprintCallable, Category = "Audio Analyzer")
+	void GetCenterFrequencies(UPARAM(DisplayName = "Sample Rate") const float InSampleRate, UPARAM(DisplayName = "Center Frequencies") TArray<float>& OutCenterFrequencies);
+
+	UFUNCTION(BlueprintCallable, Category = "Audio Analyzer")
+	UPARAM(DisplayName = "Num Center Frequencies") const int32 GetNumCenterFrequencies() const;
+
+protected:
+
+	/** Return the name of the IAudioAnalyzerFactory associated with this UAudioAnalyzer */
+	FName GetAnalyzerFactoryName() const override;
+};
\ No newline at end of file
diff --git a/Engine/Plugins/Runtime/AudioSynesthesia/Source/AudioSynesthesia/Private/SynesthesiaSpectrumAnalysis.cpp b/Engine/Plugins/Runtime/AudioSynesthesia/Source/AudioSynesthesia/Private/SynesthesiaSpectrumAnalysis.cpp
new file mode 100644
--- /dev/null
+++ b/Engine/Plugins/Runtime/AudioSynesthesia/Source/AudioSynesthesia/Private/SynesthesiaSpectrumAnalysis.cpp
@@ -0,0 +1,208 @@
+// Copyright Epic Games, Inc. All Rights Reserved.
+
+#include "SynesthesiaSpectrumAnalysis.h"
+
+namespace 
+{
+	// Convert EFFTSize to int32
+	const int32 BPFFTSizeEnumToInt(EFFTSize InFFTSizeEnum)
+	{
+		switch (InFFTSizeEnum)
+		{
+		case EFFTSize::Min:
+			return 64;
+			break;
+		case EFFTSize::Small:
+			return 256;
+			break;
+		case EFFTSize::Medium:
+			return 512;
+			break;
+		case EFFTSize::Large:
+			return 1024;
+			break;
+		case EFFTSize::VeryLarge:
+			return 2048;
+			break;
+		case EFFTSize::Max:
+			return 4096;
+			break;
+
+		case EFFTSize::DefaultSize:
+		default:
+			return 512;
+			break;
+		};
+	}
+
+	// EAudioSpectrumType to Audio::ESynesthesiaSpectrumType
+	const Audio::ESynesthesiaSpectrumType BPSpectrumTypeToSynesthesiaSpectrumType(EAudioSpectrumType InSpectrumType)
+	{
+		switch (InSpectrumType)
+		{
+		case EAudioSpectrumType::MagnitudeSpectrum:
+			return Audio::ESynesthesiaSpectrumType::MagnitudeSpectrum;
+			break;
+
+		case EAudioSpectrumType::Decibel:
+			return Audio::ESynesthesiaSpectrumType::Decibel;
+			break;
+
+		case EAudioSpectrumType::PowerSpectrum:
+		default:
+			return Audio::ESynesthesiaSpectrumType::PowerSpectrum;
+			break;
+		};
+	}
+
+	// EFFTWindowType to Audio::ESynesthesiaFFTWindowType
+	const Audio::EWindowType BPWindowTypeToWindowType(EFFTWindowType InSpectrumType)
+	{
+		switch (InSpectrumType)
+		{
+		case EFFTWindowType::Hamming:
+			return Audio::EWindowType::Hamming;
+			break;
+
+		case EFFTWindowType::Hann:
+			return Audio::EWindowType::Hann;
+			break;
+
+		case EFFTWindowType::Blackman:
+			return Audio::EWindowType::Blackman;
+			break;
+
+		case EFFTWindowType::None:
+		default:
+			return Audio::EWindowType::None;
+			break;
+		};
+	}
+}
+
+TUniquePtr<Audio::IAnalyzerSettings> USynesthesiaSpectrumAnalysisSettings::GetSettings(const int32 InSampleRate, const int32 InNumChannels) const
+{
+	TUniquePtr<Audio::FSynesthesiaSpectrumAnalysisSettings> Settings = MakeUnique<Audio::FSynesthesiaSpectrumAnalysisSettings>();
+
+	Settings->AnalysisPeriod = AnalysisPeriod;
+	Settings->FFTSize = BPFFTSizeEnumToInt(FFTSize);
+	Settings->SpectrumType = BPSpectrumTypeToSynesthesiaSpectrumType(SpectrumType);
+	Settings->WindowType = BPWindowTypeToWindowType(WindowType);
+	Settings->bDownmixToMono = bDownmixToMono;
+	return Settings;
+}
+
+#if WITH_EDITOR
+FText USynesthesiaSpectrumAnalysisSettings::GetAssetActionName() const
+{
+	return NSLOCTEXT("AssetTypeActions", "AssetTypeActions_AssetSoundSynesthesiaSpectrumSettings", "Synesthesia Real-Time Settings (Spectrum)");
+}
+
+UClass* USynesthesiaSpectrumAnalysisSettings::GetSupportedClass() const
+{
+	return USynesthesiaSpectrumAnalysisSettings::StaticClass();
+}
+#endif
+
+USynesthesiaSpectrumAnalyzer::USynesthesiaSpectrumAnalyzer()
+{
+	Settings = CreateDefaultSubobject<USynesthesiaSpectrumAnalysisSettings>(TEXT("DefaultSynesthesiaSpectrumSettings"));
+}
+
+TUniquePtr<Audio::IAnalyzerSettings> USynesthesiaSpectrumAnalyzer::GetSettings(const int32 InSampleRate, const int32 InNumChannels) const
+{
+	TUniquePtr<Audio::IAnalyzerSettings> AnalyzerSettings;
+
+	if (Settings)
+	{
+		AnalyzerSettings = Settings->GetSettings(InSampleRate, InNumChannels);
+	}
+
+	return AnalyzerSettings;
+}
+
+static TArray<FSynesthesiaSpectrumResults> ConvertToBlueprintResults(USynesthesiaSpectrumAnalysisSettings* Settings, const TArray<Audio::FSynesthesiaSpectrumEntry>& InSynesthesiaSpectrumArray)
+{
+	TArray<FSynesthesiaSpectrumResults> ResultsArray;
+
+	for (const Audio::FSynesthesiaSpectrumEntry& SynesthesiaSpectrumEntry : InSynesthesiaSpectrumArray)
+	{
+		FSynesthesiaSpectrumResults NewResults;
+
+		NewResults.SpectrumValues = SynesthesiaSpectrumEntry.SpectrumValues;
+		NewResults.TimeSeconds = SynesthesiaSpectrumEntry.Timestamp;
+
+		ResultsArray.Add(NewResults);
+	}
+
+	// Sort by time priority (lowest priority/latest first).
+	ResultsArray.Sort([](const FSynesthesiaSpectrumResults& A, const FSynesthesiaSpectrumResults& B) 		
+	{ 
+		return A.TimeSeconds < B.TimeSeconds;
+	});
+
+	return MoveTemp(ResultsArray);
+}
+
+void USynesthesiaSpectrumAnalyzer::BroadcastResults()
+{
+	TUniquePtr<const Audio::FSynesthesiaSpectrumResult> SynesthesiaSpectrumResults = GetResults<Audio::FSynesthesiaSpectrumResult>();
+	if (!SynesthesiaSpectrumResults.IsValid())
+	{
+		return;
+	}
+
+	int32 NumChannels = SynesthesiaSpectrumResults->GetNumChannels();
+
+	if (NumChannels <= 0)
+	{
+		return;
+	}
+
+	bool bIsOnSpectrumResultsBound = OnSpectrumResults.IsBound() || OnSpectrumResultsNative.IsBound();
+	bool bIsOnLatestSpectrumResultsBound = OnLatestSpectrumResults.IsBound() || OnLatestSpectrumResultsNative.IsBound();
+
+	if (bIsOnSpectrumResultsBound || bIsOnLatestSpectrumResultsBound)
+	{
+		for (int32 ChannelIndex = 0; ChannelIndex < NumChannels; ++ChannelIndex)
+		{
+			const TArray<Audio::FSynesthesiaSpectrumEntry>& SynesthesiaSpectrumArray = SynesthesiaSpectrumResults->GetChannelSpectrumArray(ChannelIndex);
+			if (SynesthesiaSpectrumArray.Num() > 0)
+			{
+				// Sorts results by timestamp
+				TArray<FSynesthesiaSpectrumResults> Results = ConvertToBlueprintResults(Settings, SynesthesiaSpectrumArray);
+				check(Results.Num() > 0);
+
+				OnSpectrumResults.Broadcast(ChannelIndex, Results);
+				OnSpectrumResultsNative.Broadcast(this, ChannelIndex, Results);
+
+				OnLatestSpectrumResults.Broadcast(ChannelIndex, Results[Results.Num() - 1]);
+				OnLatestSpectrumResultsNative.Broadcast(this, ChannelIndex, Results[Results.Num() - 1]);
+			}
+		}
+	}
+}
+
+void USynesthesiaSpectrumAnalyzer::GetCenterFrequencies(const float InSampleRate, TArray<float>& OutCenterFrequencies)
+{
+	const int32 NumFrequencies = GetNumCenterFrequencies();
+	if (NumFrequencies == OutCenterFrequencies.Num())
+	{
+		float BinSize = InSampleRate / NumFrequencies / 2; // with Nyquist
+		for (int i = 0; i < NumFrequencies; ++i)
+		{
+			OutCenterFrequencies[i] = i * BinSize;
+		}
+	}
+}
+
+const int32 USynesthesiaSpectrumAnalyzer::GetNumCenterFrequencies() const
+{
+	return BPFFTSizeEnumToInt(Settings->FFTSize) / 2;
+}
+
+FName USynesthesiaSpectrumAnalyzer::GetAnalyzerFactoryName() const
+{
+	static const FName FactoryName(TEXT("SpectrumAnalysisFactory"));
+	return FactoryName;
+}
diff --git a/Engine/Plugins/Runtime/AudioSynesthesia/Source/AudioSynesthesiaCore/Private/AudioSynesthesiaCoreModule.cpp b/Engine/Plugins/Runtime/AudioSynesthesia/Source/AudioSynesthesiaCore/Private/AudioSynesthesiaCoreModule.cpp
--- a/Engine/Plugins/Runtime/AudioSynesthesia/Source/AudioSynesthesiaCore/Private/AudioSynesthesiaCoreModule.cpp
+++ b/Engine/Plugins/Runtime/AudioSynesthesia/Source/AudioSynesthesiaCore/Private/AudioSynesthesiaCoreModule.cpp
@@ -12,6 +12,7 @@
 #include "Modules/ModuleInterface.h"
 #include "Modules/ModuleManager.h"
 #include "OnsetNRTFactory.h"
+#include "SynesthesiaSpectrumAnalysisFactory.h"
 
 DEFINE_LOG_CATEGORY(LogAudioSynesthesiaCore);
 
@@ -23,6 +24,8 @@ namespace Audio
 			void StartupModule()
 			{
 				LLM_SCOPE_BYNAME(TEXT("Audio/AudioAnalysis"));
+				FModuleManager::Get().LoadModuleChecked(TEXT("SignalProcessing"));
+
 				// Register factories on startup
 				IModularFeatures::Get().RegisterModularFeature(FLoudnessNRTFactory::GetModularFeatureName(), &LoudnessNRTFactory);
 				IModularFeatures::Get().RegisterModularFeature(FConstantQNRTFactory::GetModularFeatureName(), &ConstantQNRTFactory);
@@ -30,6 +33,7 @@ namespace Audio
 
 				IModularFeatures::Get().RegisterModularFeature(FLoudnessFactory::GetModularFeatureName(), &LoudnessFactory);
 				IModularFeatures::Get().RegisterModularFeature(FMeterFactory::GetModularFeatureName(), &MeterFactory);
+				IModularFeatures::Get().RegisterModularFeature(FSynesthesiaSpectrumAnalysisFactory::GetModularFeatureName(), &SpectralAnalysisFactory);
 			}
 
 			void ShutdownModule()
@@ -42,6 +46,7 @@ namespace Audio
 
 				IModularFeatures::Get().UnregisterModularFeature(FLoudnessFactory::GetModularFeatureName(), &LoudnessFactory);
 				IModularFeatures::Get().UnregisterModularFeature(FMeterFactory::GetModularFeatureName(), &MeterFactory);
+				IModularFeatures::Get().UnregisterModularFeature(FSynesthesiaSpectrumAnalysisFactory::GetModularFeatureName(), &SpectralAnalysisFactory);
 			}
 
 		private:
@@ -50,7 +55,8 @@ namespace Audio
 			FOnsetNRTFactory OnsetNRTFactory;
 
 			FLoudnessFactory LoudnessFactory;
-			FMeterFactory MeterFactory;
+			FMeterFactory MeterFactory; 
+			FSynesthesiaSpectrumAnalysisFactory SpectralAnalysisFactory;
 	};
 }
 
diff --git a/Engine/Plugins/Runtime/AudioSynesthesia/Source/AudioSynesthesiaCore/Private/SynesthesiaSpectrumAnalysisFactory.cpp b/Engine/Plugins/Runtime/AudioSynesthesia/Source/AudioSynesthesiaCore/Private/SynesthesiaSpectrumAnalysisFactory.cpp
new file mode 100644
--- /dev/null
+++ b/Engine/Plugins/Runtime/AudioSynesthesia/Source/AudioSynesthesiaCore/Private/SynesthesiaSpectrumAnalysisFactory.cpp
@@ -0,0 +1,161 @@
+// Copyright Epic Games, Inc. All Rights Reserved.
+
+#include "SynesthesiaSpectrumAnalysisFactory.h"
+#include "AudioSynesthesiaCustomVersion.h"
+#include "DSP/DeinterleaveView.h"
+#include "DSP/FloatArrayMath.h"
+#include "DSP/SlidingWindow.h"
+
+namespace Audio
+{
+	void FSynesthesiaSpectrumResult::Add(FSynesthesiaSpectrumEntry&& InEntry)
+	{
+		// Store Spectrum data in appropriate channel
+		TArray<FSynesthesiaSpectrumEntry>& SpectrumArray = ChannelSpectrumArrays.FindOrAdd(InEntry.Channel);
+		SpectrumArray.Add(MoveTemp(InEntry));
+	}
+
+	const TArray<FSynesthesiaSpectrumEntry>& FSynesthesiaSpectrumResult::GetChannelSpectrumArray(int32 ChannelIdx) const
+	{
+		return ChannelSpectrumArrays[ChannelIdx];
+	}
+
+	int32 FSynesthesiaSpectrumResult::GetNumChannels() const 
+	{
+		return ChannelSpectrumArrays.Num();
+	}
+
+	FSynesthesiaSpectrumAnalysisWorker::FSynesthesiaSpectrumAnalysisWorker(const FAnalyzerParameters& InParams, const FSynesthesiaSpectrumAnalysisSettings& InAnalyzerSettings)
+		: NumChannels(InParams.NumChannels)
+		, SampleRate(InParams.SampleRate)
+	{
+		check(NumChannels > 0);
+		check(SampleRate > 0);
+		check(FMath::IsPowerOfTwo(InAnalyzerSettings.FFTSize));
+		
+		NumWindowFrames = InAnalyzerSettings.FFTSize;
+		NumWindowSamples = InAnalyzerSettings.FFTSize * InParams.NumChannels;
+
+		NumHopFrames = FMath::CeilToInt(InAnalyzerSettings.AnalysisPeriod * (float)InParams.SampleRate);
+		NumHopFrames = FMath::Max(1, NumHopFrames);
+
+		int32 NumHopSamples = NumHopFrames * NumChannels;
+
+		bDownmixToMono = InAnalyzerSettings.bDownmixToMono;
+		if (bDownmixToMono)
+		{
+			MonoBuffer.Reset(NumWindowFrames);
+			MonoBuffer.AddUninitialized(NumWindowFrames);
+		}
+
+		InternalBuffer = MakeUnique<TSlidingBuffer<float>>(NumWindowSamples, NumHopSamples);
+		SpectrumAnalyzer = MakeUnique<FSynesthesiaSpectrumAnalyzer>(InParams.SampleRate, InAnalyzerSettings);
+	}
+
+	void FSynesthesiaSpectrumAnalysisWorker::Analyze(TArrayView<const float> InAudio, IAnalyzerResult* OutResult)
+	{
+		FSynesthesiaSpectrumResult* SpectrumResult = static_cast<FSynesthesiaSpectrumResult*>(OutResult);
+		check(SpectrumResult != nullptr);
+
+		TAutoSlidingWindow<float> SlidingWindow(*InternalBuffer, InAudio, InternalWindow);
+
+		// Mono
+		if (NumChannels == 1)
+		{
+			for (const TArray<float>& Window : SlidingWindow)
+			{
+				FSynesthesiaSpectrumEntry NewEntry;
+				NewEntry.Channel = 0;
+				NewEntry.Timestamp = static_cast<float>(FrameCounter) / SampleRate;
+				NewEntry.SpectrumValues.AddZeroed(Window.Num() / 2);
+
+				SpectrumAnalyzer->ProcessAudio(Window, NewEntry.SpectrumValues);
+
+				SpectrumResult->Add(MoveTemp(NewEntry));
+				FrameCounter += NumHopFrames;
+			}
+			return;
+		}
+
+		// Multichannel
+		if (bDownmixToMono)
+		{
+			for (const TArray<float>& Window : SlidingWindow)
+			{
+				FMemory::Memset(MonoBuffer.GetData(), 0, sizeof(float) * NumWindowFrames);
+
+				// Downmix to mono
+				for (int32 FrameIndex = 0; FrameIndex < MonoBuffer.Num(); ++FrameIndex)
+				{
+					for (int32 ChannelIndex = 0; ChannelIndex < NumChannels; ++ChannelIndex)
+					{
+						MonoBuffer[FrameIndex] += Window[FrameIndex * NumChannels + ChannelIndex];
+					}
+				}
+
+				// Equal power sum. assuming incoherent signals.
+				ArrayMultiplyByConstantInPlace(MonoBuffer, 1.f / FMath::Sqrt(static_cast<float>(NumChannels)));
+
+				// Create results 
+				FSynesthesiaSpectrumEntry NewEntry;
+				NewEntry.Channel = 0; // Mono channel index
+				NewEntry.Timestamp = static_cast<float>(FrameCounter) / SampleRate;
+				NewEntry.SpectrumValues.AddZeroed(MonoBuffer.Num() / 2);
+
+				SpectrumAnalyzer->ProcessAudio(MonoBuffer, NewEntry.SpectrumValues);
+
+				SpectrumResult->Add(MoveTemp(NewEntry));
+				FrameCounter += NumHopFrames;
+			}
+		}
+		else // Calculate per channel 
+		{
+			for (const TArray<float>& Window : SlidingWindow)
+			{
+				int32 ChannelIndex = 0;
+				TAutoDeinterleaveView<float, FAudioBufferAlignedAllocator> DeinterleaveView(Window, ChannelBuffer, NumChannels);
+				for (auto Channel : DeinterleaveView)
+				{
+					// Create results 
+					FSynesthesiaSpectrumEntry NewEntry;
+					NewEntry.Channel = ChannelIndex;
+					NewEntry.Timestamp = static_cast<float>(FrameCounter) / SampleRate;
+					NewEntry.SpectrumValues.AddZeroed(Channel.Values.Num() / 2);
+					
+					SpectrumAnalyzer->ProcessAudio(Channel.Values, NewEntry.SpectrumValues);
+
+					SpectrumResult->Add(MoveTemp(NewEntry));
+					ChannelIndex++;
+				}
+				FrameCounter += NumHopFrames;
+			}
+		}
+	}
+
+	FName FSynesthesiaSpectrumAnalysisFactory::GetName() const 
+	{
+		static FName FactoryName(TEXT("SpectrumAnalysisFactory"));
+		return FactoryName;
+	}
+
+	FString FSynesthesiaSpectrumAnalysisFactory::GetTitle() const
+	{
+		return TEXT("Real-Time Spectrum Analyzer");
+	}
+
+	TUniquePtr<IAnalyzerResult> FSynesthesiaSpectrumAnalysisFactory::NewResult() const
+	{
+		TUniquePtr<FSynesthesiaSpectrumResult> Result = MakeUnique<FSynesthesiaSpectrumResult>();
+		return Result;
+	}
+
+	TUniquePtr<IAnalyzerWorker> FSynesthesiaSpectrumAnalysisFactory::NewWorker(const FAnalyzerParameters& InParams, const IAnalyzerSettings* InSettings) const
+	{
+		const FSynesthesiaSpectrumAnalysisSettings* SpectrumSettings = static_cast<const FSynesthesiaSpectrumAnalysisSettings*>(InSettings);
+
+		check(nullptr != SpectrumSettings);
+
+		return MakeUnique<FSynesthesiaSpectrumAnalysisWorker>(InParams, *SpectrumSettings);
+	}
+}
+
diff --git a/Engine/Plugins/Runtime/AudioSynesthesia/Source/AudioSynesthesiaCore/Private/SynesthesiaSpectrumAnalyzer.cpp b/Engine/Plugins/Runtime/AudioSynesthesia/Source/AudioSynesthesiaCore/Private/SynesthesiaSpectrumAnalyzer.cpp
new file mode 100644
--- /dev/null
+++ b/Engine/Plugins/Runtime/AudioSynesthesia/Source/AudioSynesthesiaCore/Private/SynesthesiaSpectrumAnalyzer.cpp
@@ -0,0 +1,108 @@
+ //Copyright Epic Games, Inc. All Rights Reserved.
+
+
+#include "SynesthesiaSpectrumAnalyzer.h"
+#include "DSP/FFTAlgorithm.h"
+#include "DSP/FloatArrayMath.h"
+
+namespace Audio
+{
+	FSynesthesiaSpectrumAnalyzer::FSynesthesiaSpectrumAnalyzer(float InSampleRate, const FSynesthesiaSpectrumAnalyzerSettings& InSettings)
+	: Settings(InSettings)
+	, SampleRate(InSampleRate)
+	, Window(InSettings.WindowType, InSettings.FFTSize, /*NumChannels=*/1, /*bIsPeriodic=*/false)
+	{
+		if (!ensure(SampleRate > 0.f))
+		{
+			SampleRate = 48000.f;
+		}
+
+		// Ranges from EFFTSize
+		static const int32 MinLog2FFTSize = 6; 
+		static const int32 MaxLog2FFTSize = 16; 
+
+		// Create FFT
+		FFFTSettings FFTSettings;
+
+		int32 Log2FFTSize = MinLog2FFTSize; 
+		while ((Settings.FFTSize > (1 << Log2FFTSize)) && (Log2FFTSize < MaxLog2FFTSize))
+		{
+			Log2FFTSize++;
+		}
+
+		FFTSettings.Log2Size = Log2FFTSize;
+		FFTSettings.bArrays128BitAligned = true;
+		FFTSettings.bEnableHardwareAcceleration = true;
+
+		FFT = FFFTFactory::NewFFTAlgorithm(FFTSettings);
+		check(FFT.IsValid());
+
+		// Calculate forward scaling factor to apply to power spectrum
+		switch (FFT->ForwardScaling())
+		{
+		case EFFTScaling::MultipliedByFFTSize:
+			FFTScaling = 1.f / (Settings.FFTSize * Settings.FFTSize);
+			break;
+
+		case EFFTScaling::MultipliedBySqrtFFTSize:
+			FFTScaling = 1.f / Settings.FFTSize;
+			break;
+
+		case EFFTScaling::DividedByFFTSize:
+			FFTScaling = Settings.FFTSize * Settings.FFTSize;
+			break;
+
+		case EFFTScaling::DividedBySqrtFFTSize:
+			FFTScaling = Settings.FFTSize;
+			break;
+
+		case EFFTScaling::None:
+		default:
+			FFTScaling = 1.f;
+			break;
+		}
+
+		WindowedBuffer.SetNumUninitialized(Settings.FFTSize);
+		FFTOutput.SetNumUninitialized(Settings.FFTSize);
+	}
+
+	void FSynesthesiaSpectrumAnalyzer::ProcessAudio(TArrayView<const float> InSampleView, TArrayView<float> OutSpectrum)
+	{
+		if (!ensure(InSampleView.Num() == Settings.FFTSize) || !ensure(OutSpectrum.Num() == Settings.FFTSize / 2))
+		{
+			return;
+		}
+		
+		// Copy buffer and apply window
+		FMemory::Memcpy(WindowedBuffer.GetData(), InSampleView.GetData(), sizeof(float) * Settings.FFTSize);
+		Window.ApplyToBuffer(WindowedBuffer.GetData());
+
+		// Perform FFT and convert to power 
+		FFT->ForwardRealToComplex(WindowedBuffer.GetData(), FFTOutput.GetData());
+		ArrayComplexToPower(FFTOutput, OutSpectrum);
+		ArrayMultiplyByConstantInPlace(OutSpectrum, FFTScaling);
+
+		switch (Settings.SpectrumType)
+		{
+			case ESynesthesiaSpectrumType::MagnitudeSpectrum:
+				ArraySqrtInPlace(OutSpectrum);
+				break;
+
+			// Case already covered by complex to power above 
+			case ESynesthesiaSpectrumType::PowerSpectrum:
+				break;
+
+			case ESynesthesiaSpectrumType::Decibel:
+				ArrayPowerToDecibelInPlace(OutSpectrum, -90.f);
+				break;
+
+			default:
+				checkf(false, TEXT("Unhandled ESynesthesiaSpectrumType"));
+		}
+	}
+
+	const FSynesthesiaSpectrumAnalyzerSettings& FSynesthesiaSpectrumAnalyzer::GetSettings() const
+	{
+		return Settings;
+	}
+}
diff --git a/Engine/Plugins/Runtime/AudioSynesthesia/Source/AudioSynesthesiaCore/Public/SynesthesiaSpectrumAnalysisFactory.h b/Engine/Plugins/Runtime/AudioSynesthesia/Source/AudioSynesthesiaCore/Public/SynesthesiaSpectrumAnalysisFactory.h
new file mode 100644
--- /dev/null
+++ b/Engine/Plugins/Runtime/AudioSynesthesia/Source/AudioSynesthesiaCore/Public/SynesthesiaSpectrumAnalysisFactory.h
@@ -0,0 +1,103 @@
+// Copyright Epic Games, Inc. All Rights Reserved.
+
+#pragma once
+
+#include "SynesthesiaSpectrumAnalyzer.h"
+#include "IAudioAnalyzerInterface.h"
+#include "DSP/SlidingWindow.h"
+
+namespace Audio
+{
+	/**
+	 * Contains settings for Spectrum analyzer.
+	 */
+	class AUDIOSYNESTHESIACORE_API FSynesthesiaSpectrumAnalysisSettings : public IAnalyzerSettings, public FSynesthesiaSpectrumAnalyzerSettings
+	{	
+	public:
+		/** Number of seconds between spectrum measurements */
+		float AnalysisPeriod = 0.01f;
+	};
+
+	/**
+	 * Holds the spectrum results per a time step for each channel
+	 */
+	struct FSynesthesiaSpectrumEntry
+	{
+		int32 Channel = 0;
+		float Timestamp = 0.0f;
+		TArray<float> SpectrumValues;
+	};
+
+	/** 
+	 * FSynesthesiaSpectrumResult contains the temporal evolution of loudness for the last buffer of audio analyzed.
+	 */
+	class AUDIOSYNESTHESIACORE_API FSynesthesiaSpectrumResult : public IAnalyzerResult
+	{
+	public:
+		FSynesthesiaSpectrumResult() {}
+
+		/** Appends an FSynesthesiaSpectrumEntry to the container. */
+		void Add(FSynesthesiaSpectrumEntry&& InEntry);
+
+		/** Returns const reference to FSynesthesiaSpectrumEntry array for individual channel. */
+		const TArray<FSynesthesiaSpectrumEntry>& GetChannelSpectrumArray(int32 ChannelIdx) const;
+
+		/** Returns the number of channels. */
+		int32 GetNumChannels() const;
+
+	private:
+		float DurationInSeconds = 0.0f;
+		TMap<int32, TArray<FSynesthesiaSpectrumEntry>> ChannelSpectrumArrays;
+	};
+
+	/**
+	 * FSynesthesiaSpectrumWorker performs Spectrum analysis on input sample buffers.
+	 */
+	class AUDIOSYNESTHESIACORE_API FSynesthesiaSpectrumAnalysisWorker : public IAnalyzerWorker
+	{
+	public:
+		/** Construct a worker */
+		FSynesthesiaSpectrumAnalysisWorker(const FAnalyzerParameters& InParams, const FSynesthesiaSpectrumAnalysisSettings& InAnalyzerSettings);
+
+		/** Analyzes input sample buffer and updates result. */
+		virtual void Analyze(TArrayView<const float> InAudio, IAnalyzerResult* OutResult) override;
+
+	private:
+
+		int32 NumChannels = 0;
+		bool bDownmixToMono = true;
+		int32 NumWindowFrames = 0;
+		int32 NumWindowSamples = 0;
+		int32 NumHopFrames = 0;
+		int32 SampleRate = 0;
+		int64 FrameCounter = 0;
+
+		FAlignedFloatBuffer MonoBuffer; 
+		FAlignedFloatBuffer ChannelBuffer;
+
+		TArray<float> InternalWindow;
+		TUniquePtr<TSlidingBuffer<float>> InternalBuffer;
+		TUniquePtr<FSynesthesiaSpectrumAnalyzer> SpectrumAnalyzer;
+	};
+
+	/**
+	 * Defines the Spectrum analyzer and creates related classes.
+	 */
+	class AUDIOSYNESTHESIACORE_API FSynesthesiaSpectrumAnalysisFactory : public IAnalyzerFactory
+	{
+		public:
+
+		/** Name of specific analyzer type. */
+		virtual FName GetName() const override;
+
+		/** Human readable name of analyzer. */
+		virtual FString GetTitle() const override;
+
+		/** Creates a new FSynesthesiaSpectrumResult */
+		virtual TUniquePtr<IAnalyzerResult> NewResult() const override;
+
+		/**  Creates a new FSynesthesiaSpectrumAnalysisWorker. This expects IAnalyzerNRTSettings to be a valid pointer to a FSynesthesiaSpectrumAnalysisSettings object. */
+		virtual TUniquePtr<IAnalyzerWorker> NewWorker(const FAnalyzerParameters& InParams, const IAnalyzerSettings* InSettings) const override;
+	};
+}
+
diff --git a/Engine/Plugins/Runtime/AudioSynesthesia/Source/AudioSynesthesiaCore/Public/SynesthesiaSpectrumAnalyzer.h b/Engine/Plugins/Runtime/AudioSynesthesia/Source/AudioSynesthesiaCore/Public/SynesthesiaSpectrumAnalyzer.h
new file mode 100644
--- /dev/null
+++ b/Engine/Plugins/Runtime/AudioSynesthesia/Source/AudioSynesthesiaCore/Public/SynesthesiaSpectrumAnalyzer.h
@@ -0,0 +1,61 @@
+// Copyright Epic Games, Inc. All Rights Reserved.
+
+#pragma once
+
+#include "DSP/AlignedBuffer.h"
+#include "DSP/AudioFFT.h"
+#include "DSP/FFTAlgorithm.h"
+
+namespace Audio
+{
+	enum class ESynesthesiaSpectrumType : uint8
+	{
+		// Spectrum frequency values are equal to magnitude of frequency.
+		MagnitudeSpectrum,
+
+		// Spectrum frequency values are equal to magnitude squared.
+		PowerSpectrum,
+
+		// Returns decibels with a noise floor of -90
+		Decibel,
+	};
+
+	struct AUDIOSYNESTHESIACORE_API FSynesthesiaSpectrumAnalyzerSettings 
+	{	
+		int32 FFTSize = 512;
+
+		ESynesthesiaSpectrumType SpectrumType = ESynesthesiaSpectrumType::PowerSpectrum;
+		
+		EWindowType WindowType = EWindowType::None;
+		
+		bool bDownmixToMono = true;
+	};
+
+	class AUDIOSYNESTHESIACORE_API FSynesthesiaSpectrumAnalyzer
+	{
+	public:
+
+		/** Construct analyzer */
+		FSynesthesiaSpectrumAnalyzer(float InSampleRate, const FSynesthesiaSpectrumAnalyzerSettings& InSettings);
+
+		/**
+		 * Calculate the spectrum values for the input samples. OutSpectrum should have (num input samples / 2) values. 
+		 */
+		void ProcessAudio(TArrayView<const float> InSampleView, TArrayView<float> OutSpectrum);
+
+		/**
+		 * Return const reference to settings used inside this analyzer.
+		 */
+		const FSynesthesiaSpectrumAnalyzerSettings& GetSettings() const;
+
+	protected:
+		FSynesthesiaSpectrumAnalyzerSettings Settings;
+		float SampleRate = 0.0f;
+		float FFTScaling = 0.0f;
+
+		FWindow Window;
+		Audio::FAlignedFloatBuffer WindowedBuffer;
+		Audio::FAlignedFloatBuffer FFTOutput;
+		TUniquePtr<IFFTAlgorithm> FFT;
+	};
+}
