Spectrum output size fix for Audio Synesthesia spectrum analyzer

#rnx

diff --git a/Engine/Plugins/Runtime/AudioSynesthesia/Source/AudioSynesthesia/Private/SynesthesiaSpectrumAnalysis.cpp b/Engine/Plugins/Runtime/AudioSynesthesia/Source/AudioSynesthesia/Private/SynesthesiaSpectrumAnalysis.cpp
--- a/Engine/Plugins/Runtime/AudioSynesthesia/Source/AudioSynesthesia/Private/SynesthesiaSpectrumAnalysis.cpp
+++ b/Engine/Plugins/Runtime/AudioSynesthesia/Source/AudioSynesthesia/Private/SynesthesiaSpectrumAnalysis.cpp
@@ -188,7 +188,8 @@ void USynesthesiaSpectrumAnalyzer::GetCenterFrequencies(const float InSampleRate
 	const int32 NumFrequencies = GetNumCenterFrequencies();
 	if (NumFrequencies == OutCenterFrequencies.Num())
 	{
-		float BinSize = InSampleRate / NumFrequencies / 2; // with Nyquist
+		// Calculate frequencies 0 to Nyquist inclusive
+		float BinSize = InSampleRate / BPFFTSizeEnumToInt(Settings->FFTSize);
 		for (int i = 0; i < NumFrequencies; ++i)
 		{
 			OutCenterFrequencies[i] = i * BinSize;
@@ -198,7 +199,7 @@ void USynesthesiaSpectrumAnalyzer::GetCenterFrequencies(const float InSampleRate
 
 const int32 USynesthesiaSpectrumAnalyzer::GetNumCenterFrequencies() const
 {
-	return BPFFTSizeEnumToInt(Settings->FFTSize) / 2;
+	return BPFFTSizeEnumToInt(Settings->FFTSize) / 2 + 1;
 }
 
 FName USynesthesiaSpectrumAnalyzer::GetAnalyzerFactoryName() const
diff --git a/Engine/Plugins/Runtime/AudioSynesthesia/Source/AudioSynesthesiaCore/Private/SynesthesiaSpectrumAnalysisFactory.cpp b/Engine/Plugins/Runtime/AudioSynesthesia/Source/AudioSynesthesiaCore/Private/SynesthesiaSpectrumAnalysisFactory.cpp
--- a/Engine/Plugins/Runtime/AudioSynesthesia/Source/AudioSynesthesiaCore/Private/SynesthesiaSpectrumAnalysisFactory.cpp
+++ b/Engine/Plugins/Runtime/AudioSynesthesia/Source/AudioSynesthesiaCore/Private/SynesthesiaSpectrumAnalysisFactory.cpp
@@ -33,6 +33,8 @@ namespace Audio
 		check(SampleRate > 0);
 		check(FMath::IsPowerOfTwo(InAnalyzerSettings.FFTSize));
 		
+		// From IFFTAlgorithm::ForwardRealToComplex required number of output values 
+		NumOutputFrames = InAnalyzerSettings.FFTSize / 2 + 1;
 		NumWindowFrames = InAnalyzerSettings.FFTSize;
 		NumWindowSamples = InAnalyzerSettings.FFTSize * InParams.NumChannels;
 
@@ -67,7 +69,7 @@ namespace Audio
 				FSynesthesiaSpectrumEntry NewEntry;
 				NewEntry.Channel = 0;
 				NewEntry.Timestamp = static_cast<float>(FrameCounter) / SampleRate;
-				NewEntry.SpectrumValues.AddZeroed(Window.Num() / 2);
+				NewEntry.SpectrumValues.AddZeroed(NumOutputFrames);
 
 				SpectrumAnalyzer->ProcessAudio(Window, NewEntry.SpectrumValues);
 
@@ -100,7 +102,7 @@ namespace Audio
 				FSynesthesiaSpectrumEntry NewEntry;
 				NewEntry.Channel = 0; // Mono channel index
 				NewEntry.Timestamp = static_cast<float>(FrameCounter) / SampleRate;
-				NewEntry.SpectrumValues.AddZeroed(MonoBuffer.Num() / 2);
+				NewEntry.SpectrumValues.AddZeroed(NumOutputFrames);
 
 				SpectrumAnalyzer->ProcessAudio(MonoBuffer, NewEntry.SpectrumValues);
 
@@ -120,7 +122,7 @@ namespace Audio
 					FSynesthesiaSpectrumEntry NewEntry;
 					NewEntry.Channel = ChannelIndex;
 					NewEntry.Timestamp = static_cast<float>(FrameCounter) / SampleRate;
-					NewEntry.SpectrumValues.AddZeroed(Channel.Values.Num() / 2);
+					NewEntry.SpectrumValues.AddZeroed(NumOutputFrames);
 					
 					SpectrumAnalyzer->ProcessAudio(Channel.Values, NewEntry.SpectrumValues);
 
diff --git a/Engine/Plugins/Runtime/AudioSynesthesia/Source/AudioSynesthesiaCore/Private/SynesthesiaSpectrumAnalyzer.cpp b/Engine/Plugins/Runtime/AudioSynesthesia/Source/AudioSynesthesiaCore/Private/SynesthesiaSpectrumAnalyzer.cpp
--- a/Engine/Plugins/Runtime/AudioSynesthesia/Source/AudioSynesthesiaCore/Private/SynesthesiaSpectrumAnalyzer.cpp
+++ b/Engine/Plugins/Runtime/AudioSynesthesia/Source/AudioSynesthesiaCore/Private/SynesthesiaSpectrumAnalyzer.cpp
@@ -61,13 +61,13 @@ namespace Audio
 			break;
 		}
 
-		WindowedBuffer.SetNumUninitialized(Settings.FFTSize);
-		FFTOutput.SetNumUninitialized(Settings.FFTSize);
+		WindowedBuffer.SetNumUninitialized(FFT->NumInputFloats());
+		FFTOutput.SetNumUninitialized(FFT->NumOutputFloats());
 	}
 
 	void FSynesthesiaSpectrumAnalyzer::ProcessAudio(TArrayView<const float> InSampleView, TArrayView<float> OutSpectrum)
 	{
-		if (!ensure(InSampleView.Num() == Settings.FFTSize) || !ensure(OutSpectrum.Num() == Settings.FFTSize / 2))
+		if (!ensure(InSampleView.Num() == FFT->NumInputFloats()) || !ensure(OutSpectrum.Num() == FFT->NumOutputFloats() / 2))
 		{
 			return;
 		}
diff --git a/Engine/Plugins/Runtime/AudioSynesthesia/Source/AudioSynesthesiaCore/Public/SynesthesiaSpectrumAnalysisFactory.h b/Engine/Plugins/Runtime/AudioSynesthesia/Source/AudioSynesthesiaCore/Public/SynesthesiaSpectrumAnalysisFactory.h
--- a/Engine/Plugins/Runtime/AudioSynesthesia/Source/AudioSynesthesiaCore/Public/SynesthesiaSpectrumAnalysisFactory.h
+++ b/Engine/Plugins/Runtime/AudioSynesthesia/Source/AudioSynesthesiaCore/Public/SynesthesiaSpectrumAnalysisFactory.h
@@ -66,6 +66,7 @@ namespace Audio
 
 		int32 NumChannels = 0;
 		bool bDownmixToMono = true;
+		int32 NumOutputFrames = 0;
 		int32 NumWindowFrames = 0;
 		int32 NumWindowSamples = 0;
 		int32 NumHopFrames = 0;
diff --git a/Engine/Plugins/Runtime/AudioSynesthesia/Source/AudioSynesthesiaCore/Public/SynesthesiaSpectrumAnalyzer.h b/Engine/Plugins/Runtime/AudioSynesthesia/Source/AudioSynesthesiaCore/Public/SynesthesiaSpectrumAnalyzer.h
--- a/Engine/Plugins/Runtime/AudioSynesthesia/Source/AudioSynesthesiaCore/Public/SynesthesiaSpectrumAnalyzer.h
+++ b/Engine/Plugins/Runtime/AudioSynesthesia/Source/AudioSynesthesiaCore/Public/SynesthesiaSpectrumAnalyzer.h
@@ -39,7 +39,7 @@ namespace Audio
 		FSynesthesiaSpectrumAnalyzer(float InSampleRate, const FSynesthesiaSpectrumAnalyzerSettings& InSettings);
 
 		/**
-		 * Calculate the spectrum values for the input samples. OutSpectrum should have (num input samples / 2) values. 
+		 * Calculate the spectrum values for the input samples. OutSpectrum should have (num input samples / 2) + 1 values. 
 		 */
 		void ProcessAudio(TArrayView<const float> InSampleView, TArrayView<float> OutSpectrum);
 
