Db/DB to dB display name fixes
#rnx

diff --git a/Engine/Plugins/Runtime/AudioSynesthesia/Source/AudioSynesthesia/Classes/ConstantQNRT.h b/Engine/Plugins/Runtime/AudioSynesthesia/Source/AudioSynesthesia/Classes/ConstantQNRT.h
--- a/Engine/Plugins/Runtime/AudioSynesthesia/Source/AudioSynesthesia/Classes/ConstantQNRT.h
+++ b/Engine/Plugins/Runtime/AudioSynesthesia/Source/AudioSynesthesia/Classes/ConstantQNRT.h
@@ -74,14 +74,13 @@ class AUDIOSYNESTHESIA_API UConstantQNRTSettings : public UAudioSynesthesiaNRTSe
 		UPROPERTY(EditAnywhere, BlueprintReadOnly, Category=AudioAnalyzer, meta = (ClampMin = "20.0", ClampMax = "20000"))
 		float StartingFrequency;
 
-        /** Total number of resulting constant Q bands. */
+		/** Total number of resulting constant Q bands. */
 		UPROPERTY(EditAnywhere, BlueprintReadOnly, Category=AudioAnalyzer, meta = (ClampMin = "1", ClampMax = "96"))
 		int32 NumBands;
 
 		/** Number of bands within an octave. */
 		UPROPERTY(EditAnywhere, BlueprintReadOnly, Category=AudioAnalyzer, meta = (ClampMin = "1", ClampMax = "24"))
-		float NumBandsPerOctave;                        
-
+		float NumBandsPerOctave;
 		/** Number of seconds between cqt measurements */
 		UPROPERTY(EditAnywhere, BlueprintReadOnly, Category=AudioAnalyzer, meta = (ClampMin = "0.01", ClampMax = "1.0"))
 		float AnalysisPeriod;
@@ -96,7 +95,7 @@ class AUDIOSYNESTHESIA_API UConstantQNRTSettings : public UAudioSynesthesiaNRTSe
 
 		/** Type of window to be applied to input audio */
 		UPROPERTY(EditAnywhere, BlueprintReadOnly, AdvancedDisplay, Category=AudioAnalyzer)
-        EFFTWindowType WindowType;
+		EFFTWindowType WindowType;
 
 		/** Type of spectrum to use. */
 		UPROPERTY(EditAnywhere, BlueprintReadOnly, AdvancedDisplay, Category=AudioAnalyzer)
@@ -104,14 +103,14 @@ class AUDIOSYNESTHESIA_API UConstantQNRTSettings : public UAudioSynesthesiaNRTSe
 		
 		/** Stretching factor to control overlap of adjacent bands. */
 		UPROPERTY(EditAnywhere, BlueprintReadOnly, AdvancedDisplay, Category=AudioAnalyzer, meta = (ClampMin = "0.01", ClampMax = "2.0"))
-		float BandWidthStretch;                         
+		float BandWidthStretch;
 		
 		/** Normalization scheme used to generate band windows. */
 		UPROPERTY(EditAnywhere, BlueprintReadOnly, AdvancedDisplay, Category=AudioAnalyzer)
 		EConstantQNormalizationEnum CQTNormalization;
 
 		/** Noise floor to use when normalizing CQT */
-		UPROPERTY(EditAnywhere, BlueprintReadOnly, AdvancedDisplay, Category = AudioAnalyzer, meta = (ClampMin = "-120.0", ClampMax="0.0"))
+		UPROPERTY(EditAnywhere, BlueprintReadOnly, AdvancedDisplay, Category = AudioAnalyzer, meta = (DisplayName = "Noise Floor (dB)", ClampMin = "-120.0", ClampMax = "0.0"))
 		float NoiseFloorDb;
 
 		/** Convert UConstantQNRTSettings to FConstantQNRTSettings */
diff --git a/Engine/Plugins/Runtime/Synthesis/Source/Synthesis/Classes/SourceEffects/SourceEffectEQ.h b/Engine/Plugins/Runtime/Synthesis/Source/Synthesis/Classes/SourceEffects/SourceEffectEQ.h
--- a/Engine/Plugins/Runtime/Synthesis/Source/Synthesis/Classes/SourceEffects/SourceEffectEQ.h
+++ b/Engine/Plugins/Runtime/Synthesis/Source/Synthesis/Classes/SourceEffects/SourceEffectEQ.h
@@ -21,7 +21,7 @@ struct SYNTHESIS_API FSourceEffectEQBand
 	float Bandwidth;
 
 	// The gain in decibels to apply to the eq band
-	UPROPERTY(EditAnywhere, BlueprintReadWrite, Category = "SourceEffect|Preset", meta = (ClampMin = "-90.0", ClampMax = "20.0", UIMin = "-90.0", UIMax = "20.0"))
+	UPROPERTY(EditAnywhere, BlueprintReadWrite, Category = "SourceEffect|Preset", meta = (DisplayName = "Gain (dB)", ClampMin = "-90.0", ClampMax = "20.0", UIMin = "-90.0", UIMax = "20.0"))
 	float GainDb;
 
 	// Whether or not the band is enabled. Allows changing bands on the fly.
diff --git a/Engine/Plugins/Runtime/Synthesis/Source/Synthesis/Classes/SourceEffects/SourceEffectFoldbackDistortion.h b/Engine/Plugins/Runtime/Synthesis/Source/Synthesis/Classes/SourceEffects/SourceEffectFoldbackDistortion.h
--- a/Engine/Plugins/Runtime/Synthesis/Source/Synthesis/Classes/SourceEffects/SourceEffectFoldbackDistortion.h
+++ b/Engine/Plugins/Runtime/Synthesis/Source/Synthesis/Classes/SourceEffects/SourceEffectFoldbackDistortion.h
@@ -12,15 +12,15 @@ struct SYNTHESIS_API FSourceEffectFoldbackDistortionSettings
 	GENERATED_USTRUCT_BODY()
 
 	// The amount of gain to add to input to allow forcing the triggering of the threshold
-	UPROPERTY(EditAnywhere, BlueprintReadWrite, Category = "SourceEffect|Preset", meta = (ClampMin = "0.0", ClampMax = "60.0", UIMin = "0.0", UIMax = "20.0"))
+	UPROPERTY(EditAnywhere, BlueprintReadWrite, Category = "SourceEffect|Preset", meta = (DisplayName = "Input Gain (dB)", ClampMin = "0.0", ClampMax = "60.0", UIMin = "0.0", UIMax = "20.0"))
 	float InputGainDb;
 
 	// If the audio amplitude is higher than this, it will fold back
-	UPROPERTY(EditAnywhere, BlueprintReadWrite, Category = "SourceEffect|Preset", meta = (ClampMin = "-90.0", ClampMax = "0.0", UIMin = "-60.0", UIMax = "0.0"))
+	UPROPERTY(EditAnywhere, BlueprintReadWrite, Category = "SourceEffect|Preset", meta = (DisplayName = "Threshold (dB)", ClampMin = "-90.0", ClampMax = "0.0", UIMin = "-60.0", UIMax = "0.0"))
 	float ThresholdDb;
 
 	// The amount of gain to apply to the output
-	UPROPERTY(EditAnywhere, BlueprintReadWrite, Category = "SourceEffect|Preset", meta = (ClampMin = "-90.0", ClampMax = "20.0", UIMin = "-60.0", UIMax = "20.0"))
+	UPROPERTY(EditAnywhere, BlueprintReadWrite, Category = "SourceEffect|Preset", meta = (DisplayName = "Output Gain (dB)", ClampMin = "-90.0", ClampMax = "20.0", UIMin = "-60.0", UIMax = "20.0"))
 	float OutputGainDb;
 
 	FSourceEffectFoldbackDistortionSettings()
diff --git a/Engine/Plugins/Runtime/Synthesis/Source/Synthesis/Classes/SourceEffects/SourceEffectMotionFilter.h b/Engine/Plugins/Runtime/Synthesis/Source/Synthesis/Classes/SourceEffects/SourceEffectMotionFilter.h
--- a/Engine/Plugins/Runtime/Synthesis/Source/Synthesis/Classes/SourceEffects/SourceEffectMotionFilter.h
+++ b/Engine/Plugins/Runtime/Synthesis/Source/Synthesis/Classes/SourceEffects/SourceEffectMotionFilter.h
@@ -39,7 +39,7 @@ enum class ESourceEffectMotionFilterModDestination : uint8
 	FilterAResonance,
 
 	// Filter output dB range between 10.0f and -96.0f. Final Filter output is clamped to +6 dB, use positive values with caution.
-	FilterAOutputVolumeDB,
+	FilterAOutputVolumeDB UMETA(DisplayName = "Filter A Output Volume (dB)"),
 
 	// Filter input frequencies range between 20.0f and 15000.0f.
 	FilterBCutoffFrequency,
@@ -48,7 +48,7 @@ enum class ESourceEffectMotionFilterModDestination : uint8
 	FilterBResonance,
 
 	// Filter output dB range between 10.0f and -96.0f. Final Filter output is clamped to +6 dB, use positive values with caution.
-	FilterBOutputVolumeDB,
+	FilterBOutputVolumeDB UMETA(DisplayName = "Filter B Output Volume (dB)"),
 
 	// Filter Mix values range from -1.0f (Filter A) and 1.0f (Filter B).
 	FilterMix,
diff --git a/Engine/Plugins/Runtime/Synthesis/Source/Synthesis/Classes/SubmixEffects/SubmixEffectConvolutionReverb.h b/Engine/Plugins/Runtime/Synthesis/Source/Synthesis/Classes/SubmixEffects/SubmixEffectConvolutionReverb.h
--- a/Engine/Plugins/Runtime/Synthesis/Source/Synthesis/Classes/SubmixEffects/SubmixEffectConvolutionReverb.h
+++ b/Engine/Plugins/Runtime/Synthesis/Source/Synthesis/Classes/SubmixEffects/SubmixEffectConvolutionReverb.h
@@ -48,7 +48,7 @@ struct SYNTHESIS_API FSubmixEffectConvolutionReverbSettings
 	bool bMixReverbOutputToOutputChannelFormat;
 
 	/* Amout of audio to be sent to rear channels in quad/surround configurations */
-	UPROPERTY(EditAnywhere, BlueprintReadWrite, Category = "SubmixEffectPreset|Surround", meta = (EditCondition = "bMixReverbOutputToOutputChannelFormat && !bBypass", ClampMin = "-60.0", UIMin = "-60.0", ClampMax = "15.0", UIMax = "15.0"))
+	UPROPERTY(EditAnywhere, BlueprintReadWrite, Category = "SubmixEffectPreset|Surround", meta = (DisplayName = "Surround Rear Channel Bleed (dB)", EditCondition = "bMixReverbOutputToOutputChannelFormat && !bBypass", ClampMin = "-60.0", UIMin = "-60.0", ClampMax = "15.0", UIMax = "15.0"))
 	float SurroundRearChannelBleedDb;
 
 	/* If true, rear channel bleed sends will have their phase inverted. */
diff --git a/Engine/Plugins/Runtime/Synthesis/Source/Synthesis/Classes/SubmixEffects/SubmixEffectMultiBandCompressor.h b/Engine/Plugins/Runtime/Synthesis/Source/Synthesis/Classes/SubmixEffects/SubmixEffectMultiBandCompressor.h
--- a/Engine/Plugins/Runtime/Synthesis/Source/Synthesis/Classes/SubmixEffects/SubmixEffectMultiBandCompressor.h
+++ b/Engine/Plugins/Runtime/Synthesis/Source/Synthesis/Classes/SubmixEffects/SubmixEffectMultiBandCompressor.h
@@ -34,15 +34,15 @@ struct SYNTHESIS_API FDynamicsBandSettings
 	float Ratio = 1.5f;
 
 	// The knee bandwidth of the compressor to use in dB
-	UPROPERTY(EditAnywhere, BlueprintReadWrite, Category = Dynamics, meta = (ClampMin = "0.0", ClampMax = "20.0", UIMin = "0.0", UIMax = "20.0"))
+	UPROPERTY(EditAnywhere, BlueprintReadWrite, Category = Dynamics, meta = (DisplayName = "Knee (dB)", ClampMin = "0.0", ClampMax = "20.0", UIMin = "0.0", UIMax = "20.0"))
 	float KneeBandwidthDb = 10.f;
 
 	// The input gain of the dynamics processor in dB
-	UPROPERTY(EditAnywhere, BlueprintReadWrite, Category = Gain, meta = (ClampMin = "-12.0", ClampMax = "20.0", UIMin = "-12.0", UIMax = "20.0"))
+	UPROPERTY(EditAnywhere, BlueprintReadWrite, Category = Gain, meta = (DisplayName = "Input Gain (dB)", ClampMin = "-12.0", ClampMax = "20.0", UIMin = "-12.0", UIMax = "20.0"))
 	float InputGainDb = 0.f;
 
 	// The output gain of the dynamics processor in dB
-	UPROPERTY(EditAnywhere, BlueprintReadWrite, Category = Gain, meta = (ClampMin = "-12.0", ClampMax = "20.0", UIMin = "-12.0", UIMax = "20.0"))
+	UPROPERTY(EditAnywhere, BlueprintReadWrite, Category = Gain, meta = (DisplayName = "Output Gain (dB)", ClampMin = "-12.0", ClampMax = "20.0", UIMin = "-12.0", UIMax = "20.0"))
 	float OutputGainDb = 0.f;
 };
 
diff --git a/Engine/Plugins/Runtime/Synthesis/Source/Synthesis/Classes/SynthComponents/EpicSynth1Component.h b/Engine/Plugins/Runtime/Synthesis/Source/Synthesis/Classes/SynthComponents/EpicSynth1Component.h
--- a/Engine/Plugins/Runtime/Synthesis/Source/Synthesis/Classes/SynthComponents/EpicSynth1Component.h
+++ b/Engine/Plugins/Runtime/Synthesis/Source/Synthesis/Classes/SynthComponents/EpicSynth1Component.h
@@ -149,7 +149,7 @@ struct SYNTHESIS_API FModularSynthPreset : public FTableRowBase
 	ESynthLFOPatchType LFO2PatchType;
 
 	// The overall gain to use for the synthesizer in dB [-90.0, 20.0]
-	UPROPERTY(EditAnywhere, BlueprintReadWrite, Category = "Synth|Preset", meta = (ClampMin = "-90.0", ClampMax = "20.0", UIMin = "-90.0", UIMax = "20.0"))
+	UPROPERTY(EditAnywhere, BlueprintReadWrite, Category = "Synth|Preset", meta = (DisplayName = "Gain (dB)", ClampMin = "-90.0", ClampMax = "20.0", UIMin = "-90.0", UIMax = "20.0"))
 	float GainDb;
 
 	// The amplitude envelope attack time (in ms) [0.0, 10000]
@@ -524,7 +524,7 @@ public:
 	void SetLFOPatch(int32 LFOIndex, ESynthLFOPatchType LFOPatchType);
 
 	// Sets the synth gain in decibels.
-	UFUNCTION(BlueprintCallable, Category = "Synth|Components|Audio")
+	UFUNCTION(BlueprintCallable, Category = "Synth|Components|Audio", meta = (DisplayName = "Set Gain (dB)"))
 	void SetGainDb(float GainDb);
 
 	// Sets the envelope attack time in msec.
diff --git a/Engine/Plugins/Runtime/Synthesis/Source/Synthesis/Public/EffectConvolutionReverb.h b/Engine/Plugins/Runtime/Synthesis/Source/Synthesis/Public/EffectConvolutionReverb.h
--- a/Engine/Plugins/Runtime/Synthesis/Source/Synthesis/Public/EffectConvolutionReverb.h
+++ b/Engine/Plugins/Runtime/Synthesis/Source/Synthesis/Public/EffectConvolutionReverb.h
@@ -50,7 +50,7 @@ public:
 	int32 SampleRate;
 
 	/* Used to account for energy added by convolution with "loud" Impulse Responses. */
-	UPROPERTY(EditAnywhere, BlueprintReadOnly, Category = SubmixEffectPreset, meta = (ClamMin = "-60.0", ClampMax = "15.0"))
+	UPROPERTY(EditAnywhere, BlueprintReadOnly, Category = SubmixEffectPreset, meta = (DisplayName = "Normalization Volume (dB)", ClampMin = "-60.0", ClampMax = "15.0"))
 	float NormalizationVolumeDb;
 
 	/* If true, impulse response channels are interpreted as true stereo which allows channel crosstalk. If false, impulse response channels are interpreted as independent channel impulses. */
diff --git a/Engine/Source/Runtime/AudioMixer/Classes/SubmixEffects/AudioMixerSubmixEffectEQ.h b/Engine/Source/Runtime/AudioMixer/Classes/SubmixEffects/AudioMixerSubmixEffectEQ.h
--- a/Engine/Source/Runtime/AudioMixer/Classes/SubmixEffects/AudioMixerSubmixEffectEQ.h
+++ b/Engine/Source/Runtime/AudioMixer/Classes/SubmixEffects/AudioMixerSubmixEffectEQ.h
@@ -24,7 +24,7 @@ struct AUDIOMIXER_API FSubmixEffectEQBand
 	UPROPERTY(EditAnywhere, BlueprintReadWrite, Category = "SubmixEffect|Preset", meta = (ClampMin = "0.0", ClampMax = "10.0", UIMin = "0.0", UIMax = "10.0"))
 	float Bandwidth;
 
-	UPROPERTY(EditAnywhere, BlueprintReadWrite, Category = "SubmixEffect|Preset", meta = (ClampMin = "-90.0", ClampMax = "20.0", UIMin = "-90.0", UIMax = "20.0"))
+	UPROPERTY(EditAnywhere, BlueprintReadWrite, Category = "SubmixEffect|Preset", meta = (DisplayName = "Gain (dB)", ClampMin = "-90.0", ClampMax = "20.0", UIMin = "-90.0", UIMax = "20.0"))
 	float GainDb;
 
 	UPROPERTY(EditAnywhere, BlueprintReadWrite, Category = "SubmixEffect|Preset")
