Fix Audio Synthesia Asset Telemetry Typo

diff --git a/Engine/Plugins/Runtime/AudioSynesthesia/Source/AudioSynesthesiaEditor/Private/AudioSynesthesiaSettingsFactory.cpp b/Engine/Plugins/Runtime/AudioSynesthesia/Source/AudioSynesthesiaEditor/Private/AudioSynesthesiaSettingsFactory.cpp
--- a/Engine/Plugins/Runtime/AudioSynesthesia/Source/AudioSynesthesiaEditor/Private/AudioSynesthesiaSettingsFactory.cpp
+++ b/Engine/Plugins/Runtime/AudioSynesthesia/Source/AudioSynesthesiaEditor/Private/AudioSynesthesiaSettingsFactory.cpp
@@ -56,7 +56,7 @@ UObject* UAudioSynesthesiaSettingsFactory::FactoryCreateNew(UClass* InClass, UOb
 	{
 		NewAudioSynesthesiaSettings = NewObject<UAudioSynesthesiaSettings>(InParent, AudioSynesthesiaSettingsClass, InName, Flags);
 
-		Audio::Analytics::RecordEvent_Usage(TEXT("AudioSynesthesia.SettingsFactoryCreated"));
+		Audio::Analytics::RecordEvent_Usage(TEXT("AudioSynesthesia.SettingsCreated"));
 	}
 	return NewAudioSynesthesiaSettings;
 }
