Added UConstantQAnalyzer real time audio analyzer, so that there now exists both a UAudioAnalyzer and a AudioSynesthesiaNRT analyzer for ConstantQ.
This is similar to there being both ULoudnessAnalyzer and ULoudnessNRT analyzers (and has similarly duplicated code between the real-time and NRT versions).
The implementation of USynesthesiaSpectrumAnalyzer was used for inspiration also.

diff --git a/Engine/Plugins/Runtime/AudioSynesthesia/Source/AudioSynesthesia/Classes/ConstantQ.h b/Engine/Plugins/Runtime/AudioSynesthesia/Source/AudioSynesthesia/Classes/ConstantQ.h
new file mode 100644
--- /dev/null
+++ b/Engine/Plugins/Runtime/AudioSynesthesia/Source/AudioSynesthesia/Classes/ConstantQ.h
@@ -0,0 +1,147 @@
+// Copyright Epic Games, Inc. All Rights Reserved.
+
+#pragma once
+
+#include "AudioSynesthesia.h"
+#include "ConstantQNRT.h"
+#include "Sound/SoundSubmix.h"
+#include "ConstantQ.generated.h"
+
+/** UConstantQSettings
+ *
+ * Settings for a UConstantQ analyzer.
+ */
+UCLASS(Blueprintable)
+class AUDIOSYNESTHESIA_API UConstantQSettings : public UAudioSynesthesiaSettings
+{
+	GENERATED_BODY()
+public:
+
+	UConstantQSettings() {}
+
+	/** Starting frequency for first bin of CQT */
+	UPROPERTY(EditAnywhere, BlueprintReadOnly, Category=AudioAnalyzer, DisplayName = "Starting Frequency (Hz)", meta = (ClampMin = "20.0", ClampMax = "20000"))
+	float StartingFrequencyHz = 40.0f;
+
+	/** Total number of resulting constant Q bands. */
+	UPROPERTY(EditAnywhere, BlueprintReadOnly, Category=AudioAnalyzer, meta = (ClampMin = "1", ClampMax = "96"))
+	int32 NumBands = 48;
+
+	/** Number of bands within an octave. */
+	UPROPERTY(EditAnywhere, BlueprintReadOnly, Category=AudioAnalyzer, meta = (ClampMin = "1", ClampMax = "24"))
+	float NumBandsPerOctave = 12.0f;
+
+	/** Number of seconds between cqt measurements */
+	UPROPERTY(EditAnywhere, BlueprintReadOnly, Category=AudioAnalyzer, DisplayName = "Analysis Period (s)", meta = (ClampMin = "0.01", ClampMax = "1.0"))
+	float AnalysisPeriodInSeconds = 0.01f;
+
+	/** If true, multichannel audio is downmixed to mono with equal amplitude scaling. If false, each channel gets it's own CQT result. */
+	UPROPERTY(EditAnywhere, BlueprintReadOnly, Category=AudioAnalyzer)
+	bool bDownmixToMono = false;
+
+	/** Size of FFT. */
+	UPROPERTY(EditAnywhere, BlueprintReadOnly, AdvancedDisplay, Category=AudioAnalyzer)
+	EConstantQFFTSizeEnum FFTSize = EConstantQFFTSizeEnum::XLarge;
+
+	/** Type of window to be applied to input audio */
+	UPROPERTY(EditAnywhere, BlueprintReadOnly, AdvancedDisplay, Category=AudioAnalyzer)
+	EFFTWindowType WindowType = EFFTWindowType::Blackman;
+
+	/** Type of spectrum to use. */
+	UPROPERTY(EditAnywhere, BlueprintReadOnly, AdvancedDisplay, Category=AudioAnalyzer)
+	EAudioSpectrumType SpectrumType = EAudioSpectrumType::PowerSpectrum;
+		
+	/** Stretching factor to control overlap of adjacent bands. */
+	UPROPERTY(EditAnywhere, BlueprintReadOnly, AdvancedDisplay, Category=AudioAnalyzer, meta = (ClampMin = "0.01", ClampMax = "2.0"))
+	float BandWidthStretch = 1.0f;
+		
+	/** Normalization scheme used to generate band windows. */
+	UPROPERTY(EditAnywhere, BlueprintReadOnly, AdvancedDisplay, Category=AudioAnalyzer)
+	EConstantQNormalizationEnum CQTNormalization = EConstantQNormalizationEnum::EqualEnergy;
+
+	/** Noise floor to use when normalizing CQT */
+	UPROPERTY(EditAnywhere, BlueprintReadOnly, AdvancedDisplay, Category=AudioAnalyzer, meta = (DisplayName = "Noise Floor (dB)", ClampMin = "-120.0", ClampMax = "0.0"))
+	float NoiseFloorDb = -60.0f;
+
+	/** Convert UConstantQSettings to FConstantQSettings */
+	TUniquePtr<Audio::IAnalyzerSettings> GetSettings(const float InSampleRate, const int32 InNumChannels) const;
+
+#if WITH_EDITOR
+	virtual FText GetAssetActionName() const override;
+
+	virtual UClass* GetSupportedClass() const override;
+#endif // WITH_EDITOR
+};
+
+/** The results of the ConstantQ analysis. */
+USTRUCT(BlueprintType)
+struct AUDIOSYNESTHESIA_API FConstantQResults
+{
+	GENERATED_USTRUCT_BODY()
+
+	// The time in seconds since analysis began of this ConstantQ analysis result
+	UPROPERTY(EditAnywhere, BlueprintReadOnly, Category = AudioAnalyzer)
+	float TimeSeconds = 0.0f;
+
+	// The spectrum values from the FFT
+	UPROPERTY(EditAnywhere, BlueprintReadOnly, Category = AudioAnalyzer)
+	TArray<float> SpectrumValues;
+};
+
+/** Delegate to receive all ConstantQ results per channel (time-stamped in an array) since last delegate call. If bDownmixToMono setting is true, results will be in channel index 0. */
+DECLARE_DYNAMIC_MULTICAST_DELEGATE_TwoParams(FOnConstantQResults, int32, ChannelIndex, const TArray<FConstantQResults>&, ConstantQResults);
+
+/** shadow delegate declaration for above */
+DECLARE_MULTICAST_DELEGATE_ThreeParams(FOnConstantQResultsNative, UConstantQAnalyzer*, int32, const TArray<FConstantQResults>&);
+
+/** Delegate to receive only the most recent overall ConstantQ result per channel. If bDownmixToMono setting is true, results will be in channel index 0.*/
+DECLARE_DYNAMIC_MULTICAST_DELEGATE_TwoParams(FOnLatestConstantQResults, int32, ChannelIndex, const FConstantQResults&, LatestConstantQResults);
+
+/** shadow delegate declaration for above */
+DECLARE_MULTICAST_DELEGATE_ThreeParams(FOnLatestConstantQResultsNative, UConstantQAnalyzer*, int32, const FConstantQResults&);
+
+/** UConstantQAnalyzer
+ *
+ * UConstantQAnalyzer calculates the temporal evolution of constant q transform for a given
+ * audio bus in real-time. ConstantQ is available for individual channels or the overall audio bus.
+ */
+UCLASS(Blueprintable)
+class AUDIOSYNESTHESIA_API UConstantQAnalyzer : public UAudioAnalyzer
+{
+	GENERATED_BODY()
+public:
+	UConstantQAnalyzer();
+
+	/** The settings for the audio analyzer.  */
+	UPROPERTY(EditAnywhere, BlueprintReadWrite, Category = AudioAnalyzer)
+	TObjectPtr<UConstantQSettings> Settings;
+
+	/** Delegate to receive all Spectrum results, per-channel, since last delegate call. If bDownmixToMono setting is true, results will be in channel index 0. */
+	UPROPERTY(BlueprintAssignable)
+	FOnConstantQResults OnConstantQResults;
+
+	FOnConstantQResultsNative OnConstantQResultsNative;
+
+	/** Delegate to receive the latest per-channel Spectrum results. If bDownmixToMono setting is true, results will be in channel index 0. */
+	UPROPERTY(BlueprintAssignable)
+	FOnLatestConstantQResults OnLatestConstantQResults;
+
+	FOnLatestConstantQResultsNative OnLatestConstantQResultsNative;
+
+	/** Convert UConstantQSettings to FConstantQSettings */
+	TUniquePtr<Audio::IAnalyzerSettings> GetSettings(const int32 InSampleRate, const int32 InNumChannels) const override;
+
+	/** Broadcasts results to any delegates if hooked up. */
+	void BroadcastResults() override;
+
+	UFUNCTION(BlueprintCallable, Category = "Audio Analyzer")
+	void GetCenterFrequencies(UPARAM(DisplayName = "Center Frequencies") TArray<float>& OutCenterFrequencies);
+
+	UFUNCTION(BlueprintCallable, Category = "Audio Analyzer")
+	UPARAM(DisplayName = "Num Center Frequencies") const int32 GetNumCenterFrequencies() const;
+
+protected:
+
+	/** Return the name of the IAudioAnalyzerFactory associated with this UAudioAnalyzer */
+	FName GetAnalyzerFactoryName() const override;
+};
\ No newline at end of file
diff --git a/Engine/Plugins/Runtime/AudioSynesthesia/Source/AudioSynesthesia/Classes/ConstantQNRT.h b/Engine/Plugins/Runtime/AudioSynesthesia/Source/AudioSynesthesia/Classes/ConstantQNRT.h
--- a/Engine/Plugins/Runtime/AudioSynesthesia/Source/AudioSynesthesia/Classes/ConstantQNRT.h
+++ b/Engine/Plugins/Runtime/AudioSynesthesia/Source/AudioSynesthesia/Classes/ConstantQNRT.h
@@ -33,31 +33,31 @@ UENUM(BlueprintType)
 enum class EConstantQFFTSizeEnum : uint8
 {
 	// 64
-	Min,
+	Min UMETA(DisplayName = "64"),
 	
 	// 128
-	XXSmall,
+	XXSmall UMETA(DisplayName = "128"),
 
 	// 256
-	XSmall,
+	XSmall UMETA(DisplayName = "256"),
 
 	// 512
-	Small,
+	Small UMETA(DisplayName = "512"),
 
 	// 1024
-	Medium,
+	Medium UMETA(DisplayName = "1024"),
 
 	// 2048
-	Large,
+	Large UMETA(DisplayName = "2048"),
 
 	// 4096
-	XLarge,
+	XLarge UMETA(DisplayName = "4096"),
 
 	// 8192
-	XXLarge,
+	XXLarge UMETA(DisplayName = "8192"),
 
 	// 16384
-	Max
+	Max UMETA(DisplayName = "16384")
 };
 										 
 /** UConstantQNRTSettings
diff --git a/Engine/Plugins/Runtime/AudioSynesthesia/Source/AudioSynesthesia/Private/ConstantQ.cpp b/Engine/Plugins/Runtime/AudioSynesthesia/Source/AudioSynesthesia/Private/ConstantQ.cpp
new file mode 100644
--- /dev/null
+++ b/Engine/Plugins/Runtime/AudioSynesthesia/Source/AudioSynesthesia/Private/ConstantQ.cpp
@@ -0,0 +1,225 @@
+// Copyright Epic Games, Inc. All Rights Reserved.
+
+#include "ConstantQ.h"
+#include "ConstantQFactory.h"
+
+#include UE_INLINE_GENERATED_CPP_BY_NAME(ConstantQ)
+
+
+namespace ConstantQPrivate
+{
+	/** Convenience template function for converting from blueprint types to audio types */
+	template<typename InType, typename OutType>
+	OutType ConvertType(const TMap<InType, OutType>& InMap, const InType& InValue, const OutType& InDefaultValue)
+	{
+		checkf(InMap.Contains(InValue), TEXT("Unhandled value conversion"));
+
+		if (InMap.Contains(InValue))
+		{
+			return InMap[InValue];
+		}
+
+		return InDefaultValue;
+	}
+
+
+	// EConstantQFFTSizeEnum to int
+	int32 BPConstantQFFTSizeEnumToInt(const EConstantQFFTSizeEnum InFFTSizeEnum)
+	{
+		static const TMap<EConstantQFFTSizeEnum, int32> FFTSizeMap =
+		{
+			{EConstantQFFTSizeEnum::Min, 		64},
+			{EConstantQFFTSizeEnum::XXSmall, 	128},
+			{EConstantQFFTSizeEnum::XSmall, 	256},
+			{EConstantQFFTSizeEnum::Small,		512},
+			{EConstantQFFTSizeEnum::Medium, 	1024},
+			{EConstantQFFTSizeEnum::Large,		2048},
+			{EConstantQFFTSizeEnum::XLarge, 	4096},
+			{EConstantQFFTSizeEnum::XXLarge, 	8192},
+			{EConstantQFFTSizeEnum::Max, 		16384}
+		};
+
+		return ConvertType(FFTSizeMap, InFFTSizeEnum, 2048);
+
+	}
+
+	// EFFTWindwoType to Audio::EWindowType
+	Audio::EWindowType BPWindowTypeToAudioWindowType(const EFFTWindowType InWindowType)
+	{
+		static const TMap<EFFTWindowType, Audio::EWindowType> WindowTypeMap =
+		{
+			{EFFTWindowType::None, 		Audio::EWindowType::None},
+			{EFFTWindowType::Hamming, 	Audio::EWindowType::Hamming},
+			{EFFTWindowType::Hann, 		Audio::EWindowType::Hann},
+			{EFFTWindowType::Blackman,	Audio::EWindowType::Blackman}
+		};
+
+		return ConvertType(WindowTypeMap, InWindowType, Audio::EWindowType::None);
+	}
+
+	// ESpectrumType to Audio::ESpectrumType
+	Audio::ESpectrumType BPSpectrumTypeToAudioSpectrumType(const EAudioSpectrumType InSpectrumType)
+	{
+		static const TMap<EAudioSpectrumType, Audio::ESpectrumType> SpectrumTypeMap =
+		{
+			{EAudioSpectrumType::MagnitudeSpectrum,	Audio::ESpectrumType::MagnitudeSpectrum},
+			{EAudioSpectrumType::PowerSpectrum, 	Audio::ESpectrumType::PowerSpectrum}
+		};
+
+		return ConvertType(SpectrumTypeMap, InSpectrumType, Audio::ESpectrumType::PowerSpectrum);
+	}
+
+	// EContantQNormalizationEnum to Audio::EPseudoConstantQNormalization
+	Audio::EPseudoConstantQNormalization BPCQTNormalizationToAudioCQTNormalization(const EConstantQNormalizationEnum InNormalization)
+	{
+		static const TMap<EConstantQNormalizationEnum, Audio::EPseudoConstantQNormalization> NormalizationMap =
+		{
+			{EConstantQNormalizationEnum::EqualAmplitude, 		Audio::EPseudoConstantQNormalization::EqualAmplitude},
+			{EConstantQNormalizationEnum::EqualEuclideanNorm,	Audio::EPseudoConstantQNormalization::EqualEuclideanNorm},
+			{EConstantQNormalizationEnum::EqualEnergy, 			Audio::EPseudoConstantQNormalization::EqualEnergy}
+		};
+
+		return ConvertType(NormalizationMap, InNormalization, Audio::EPseudoConstantQNormalization::EqualEnergy);
+	}
+}
+
+
+/** Convert UConstantQSettings to FConstantQSettings */
+TUniquePtr<Audio::IAnalyzerSettings> UConstantQSettings::GetSettings(const float InSampleRate, const int32 InNumChannels) const
+{
+	using namespace ConstantQPrivate;
+
+	TUniquePtr<Audio::FConstantQSettings> Settings = MakeUnique<Audio::FConstantQSettings>();
+
+	Settings->AnalysisPeriodInSeconds = AnalysisPeriodInSeconds;
+	Settings->bDownmixToMono = bDownmixToMono;
+	Settings->FFTSize = BPConstantQFFTSizeEnumToInt(FFTSize);
+	Settings->WindowType = BPWindowTypeToAudioWindowType(WindowType);
+	Settings->SpectrumType = BPSpectrumTypeToAudioSpectrumType(SpectrumType);
+	Settings->NumBands = NumBands;
+	Settings->NumBandsPerOctave = NumBandsPerOctave;
+	Settings->KernelLowestCenterFreq = StartingFrequencyHz;
+	Settings->BandWidthStretch = BandWidthStretch;
+	Settings->Normalization = BPCQTNormalizationToAudioCQTNormalization(CQTNormalization);
+	Settings->Scaling = (SpectrumType == EAudioSpectrumType::Decibel) ? Audio::EConstantQScaling::Decibel : Audio::EConstantQScaling::Linear;
+
+	return Settings;
+}
+
+#if WITH_EDITOR
+FText UConstantQSettings::GetAssetActionName() const
+{
+	return NSLOCTEXT("AssetTypeActions", "AssetTypeActions_AssetSoundSynesthesiaConstantQSettings", "Synesthesia Real-Time Settings (ConstantQ)");
+}
+
+UClass* UConstantQSettings::GetSupportedClass() const
+{
+	return UConstantQNRTSettings::StaticClass();
+}
+#endif // WITH_EDITOR
+
+UConstantQAnalyzer::UConstantQAnalyzer()
+{
+	Settings = CreateDefaultSubobject<UConstantQSettings>(TEXT("DefaultConstantQSettings"));
+}
+
+TUniquePtr<Audio::IAnalyzerSettings> UConstantQAnalyzer::GetSettings(const int32 InSampleRate, const int32 InNumChannels) const
+{
+	TUniquePtr<Audio::IAnalyzerSettings> AnalyzerSettings;
+
+	if (Settings)
+	{
+		AnalyzerSettings = Settings->GetSettings(InSampleRate, InNumChannels);
+	}
+
+	return AnalyzerSettings;
+}
+
+static TArray<FConstantQResults> ConvertToBlueprintResults(UConstantQSettings* Settings, const TArray<Audio::FConstantQFrame>& InConstantQFrameArray)
+{
+	TArray<FConstantQResults> ResultsArray;
+
+	for (const Audio::FConstantQFrame& ConstantQFrame : InConstantQFrameArray)
+	{
+		FConstantQResults NewResults;
+
+		NewResults.SpectrumValues = ConstantQFrame.Spectrum;
+		NewResults.TimeSeconds = ConstantQFrame.Timestamp;
+
+		ResultsArray.Add(NewResults);
+	}
+
+	// Sort by time priority (lowest priority/latest first).
+	ResultsArray.Sort([](const FConstantQResults& A, const FConstantQResults& B)
+		{
+			return A.TimeSeconds < B.TimeSeconds;
+		});
+
+	return MoveTemp(ResultsArray);
+}
+
+void UConstantQAnalyzer::BroadcastResults()
+{
+	TUniquePtr<const Audio::FConstantQResult> ConstantQResults = GetResults<Audio::FConstantQResult>();
+	if (!ConstantQResults.IsValid())
+	{
+		return;
+	}
+
+	const int32 NumChannels = ConstantQResults->GetNumChannels();
+
+	if (NumChannels <= 0)
+	{
+		return;
+	}
+
+	bool bIsOnConstantQResultsBound = OnConstantQResults.IsBound() || OnConstantQResultsNative.IsBound();
+	bool bIsOnLatestConstantQResultsBound = OnLatestConstantQResults.IsBound() || OnLatestConstantQResultsNative.IsBound();
+
+	if (bIsOnConstantQResultsBound || bIsOnLatestConstantQResultsBound)
+	{
+		for (int32 ChannelIndex = 0; ChannelIndex < NumChannels; ++ChannelIndex)
+		{
+			const TArray<Audio::FConstantQFrame>& ConstantQFrameArray = ConstantQResults->GetFramesForChannel(ChannelIndex);
+			if (ConstantQFrameArray.Num() > 0)
+			{
+				// Sorts results by timestamp
+				TArray<FConstantQResults> Results = ConvertToBlueprintResults(Settings, ConstantQFrameArray);
+				ensure(Results.Num() > 0);
+
+				OnConstantQResults.Broadcast(ChannelIndex, Results);
+				OnConstantQResultsNative.Broadcast(this, ChannelIndex, Results);
+
+				OnLatestConstantQResults.Broadcast(ChannelIndex, Results[Results.Num() - 1]);
+				OnLatestConstantQResultsNative.Broadcast(this, ChannelIndex, Results[Results.Num() - 1]);
+			}
+		}
+	}
+}
+
+void UConstantQAnalyzer::GetCenterFrequencies(TArray<float>& OutCenterFrequencies)
+{
+	const int32 NumFrequencies = GetNumCenterFrequencies();
+	if (NumFrequencies == OutCenterFrequencies.Num())
+	{
+		const float PitchStep = FMath::Pow(2.0f, 1.0f / Settings->NumBandsPerOctave);
+
+		float CurrentFrequency = Settings->StartingFrequencyHz;
+		for (int32 i = 0; i < NumFrequencies; ++i)
+		{
+			OutCenterFrequencies[i] = CurrentFrequency;
+			CurrentFrequency *= PitchStep;
+		}
+	}
+}
+
+const int32 UConstantQAnalyzer::GetNumCenterFrequencies() const
+{
+	return Settings->NumBands;
+}
+
+FName UConstantQAnalyzer::GetAnalyzerFactoryName() const
+{
+	static const FName FactoryName(TEXT("ConstantQFactory"));
+	return FactoryName;
+}
diff --git a/Engine/Plugins/Runtime/AudioSynesthesia/Source/AudioSynesthesiaCore/Private/AudioSynesthesiaCoreModule.cpp b/Engine/Plugins/Runtime/AudioSynesthesia/Source/AudioSynesthesiaCore/Private/AudioSynesthesiaCoreModule.cpp
--- a/Engine/Plugins/Runtime/AudioSynesthesia/Source/AudioSynesthesiaCore/Private/AudioSynesthesiaCoreModule.cpp
+++ b/Engine/Plugins/Runtime/AudioSynesthesia/Source/AudioSynesthesiaCore/Private/AudioSynesthesiaCoreModule.cpp
@@ -2,6 +2,7 @@
 #include "AudioSynesthesiaCoreModule.h"
 
 #include "AudioSynesthesiaCoreLog.h"
+#include "ConstantQFactory.h"
 #include "ConstantQNRTFactory.h"
 #include "Features/IModularFeatures.h"
 #include "HAL/LowLevelMemTracker.h"
@@ -30,6 +31,7 @@ namespace Audio
 				IModularFeatures::Get().RegisterModularFeature(FOnsetNRTFactory::GetModularFeatureName(), &OnsetNRTFactory);
 
 				IModularFeatures::Get().RegisterModularFeature(FLoudnessFactory::GetModularFeatureName(), &LoudnessFactory);
+				IModularFeatures::Get().RegisterModularFeature(FConstantQFactory::GetModularFeatureName(), &ConstantQFactory);
 				IModularFeatures::Get().RegisterModularFeature(FMeterFactory::GetModularFeatureName(), &MeterFactory);
 				IModularFeatures::Get().RegisterModularFeature(FSynesthesiaSpectrumAnalysisFactory::GetModularFeatureName(), &SpectralAnalysisFactory);
 			}
@@ -43,6 +45,7 @@ namespace Audio
 				IModularFeatures::Get().UnregisterModularFeature(FOnsetNRTFactory::GetModularFeatureName(), &OnsetNRTFactory);
 
 				IModularFeatures::Get().UnregisterModularFeature(FLoudnessFactory::GetModularFeatureName(), &LoudnessFactory);
+				IModularFeatures::Get().UnregisterModularFeature(FConstantQFactory::GetModularFeatureName(), &ConstantQFactory);
 				IModularFeatures::Get().UnregisterModularFeature(FMeterFactory::GetModularFeatureName(), &MeterFactory);
 				IModularFeatures::Get().UnregisterModularFeature(FSynesthesiaSpectrumAnalysisFactory::GetModularFeatureName(), &SpectralAnalysisFactory);
 			}
@@ -53,6 +56,7 @@ namespace Audio
 			FOnsetNRTFactory OnsetNRTFactory;
 
 			FLoudnessFactory LoudnessFactory;
+			FConstantQFactory ConstantQFactory;
 			FMeterFactory MeterFactory; 
 			FSynesthesiaSpectrumAnalysisFactory SpectralAnalysisFactory;
 	};
diff --git a/Engine/Plugins/Runtime/AudioSynesthesia/Source/AudioSynesthesiaCore/Private/ConstantQFactory.cpp b/Engine/Plugins/Runtime/AudioSynesthesia/Source/AudioSynesthesiaCore/Private/ConstantQFactory.cpp
new file mode 100644
--- /dev/null
+++ b/Engine/Plugins/Runtime/AudioSynesthesia/Source/AudioSynesthesiaCore/Private/ConstantQFactory.cpp
@@ -0,0 +1,148 @@
+// Copyright Epic Games, Inc. All Rights Reserved.
+
+#include "ConstantQFactory.h"
+#include "DSP/DeinterleaveView.h"
+#include "DSP/SlidingWindow.h"
+
+namespace Audio 
+{
+	void FConstantQResult::AddFrame(const int32 InChannelIndex, const float InTimestamp, TArrayView<const float> InSpectrum)
+	{
+		TArray<FConstantQFrame>& FrameArray = ChannelCQTFrames.FindOrAdd(InChannelIndex);
+		FrameArray.Emplace(InChannelIndex, InTimestamp, InSpectrum);
+	}
+
+	const TArray<FConstantQFrame>& FConstantQResult::GetFramesForChannel(const int32 InChannelIndex) const
+	{
+		return ChannelCQTFrames[InChannelIndex];
+	}
+
+	int32 FConstantQResult::GetNumChannels() const 
+	{
+		return ChannelCQTFrames.Num();
+	}
+
+
+	FConstantQWorker::FConstantQWorker(const FAnalyzerParameters& InParams, const FConstantQSettings& InAnalyzerSettings)
+		: NumChannels(InParams.NumChannels)
+		, SampleRate(InParams.SampleRate)
+		, NumWindowFrames(InAnalyzerSettings.FFTSize)
+		, NumWindowSamples(InAnalyzerSettings.FFTSize * InParams.NumChannels)
+		, bDownmixToMono(InAnalyzerSettings.bDownmixToMono)
+	{
+		check(InParams.NumChannels > 0);
+		check(InParams.SampleRate > 0.f);
+
+		ensure(InAnalyzerSettings.AnalysisPeriodInSeconds > 0.f);
+		ensure(FMath::IsPowerOfTwo(InAnalyzerSettings.FFTSize));
+
+
+		NumHopFrames = FMath::Max(1, FMath::RoundToInt(InAnalyzerSettings.AnalysisPeriodInSeconds * InParams.SampleRate));
+		NumHopSamples = NumHopFrames * NumChannels;
+
+		if (bDownmixToMono)
+		{
+			MonoBuffer.Reset(NumWindowFrames);
+			MonoBuffer.AddUninitialized(NumWindowFrames);
+			// equal power mixing
+			MonoScaling = 1.f / FMath::Sqrt(static_cast<float>(NumChannels));
+		}
+
+		SlidingBuffer = MakeUnique<TSlidingBuffer<float> >(NumWindowSamples, NumHopSamples);
+
+		ConstantQAnalyzer = MakeUnique<FConstantQAnalyzer>(InAnalyzerSettings, SampleRate);
+	}
+
+	void FConstantQWorker::Analyze(TArrayView<const float> InAudio, IAnalyzerResult* OutResult)
+	{
+		const bool bDoFlush = false;
+
+		AnalyzeMultichannel(InAudio, OutResult, bDoFlush);
+	}
+
+	void FConstantQWorker::AnalyzeMultichannel(TArrayView<const float> InAudio, IAnalyzerResult* OutResult, const bool bDoFlush)
+	{
+		// Assume that outer layers ensured that this is of correct type.
+		FConstantQResult* ConstantQResult = static_cast<FConstantQResult*>(OutResult);
+
+		check(nullptr != ConstantQResult);
+
+		TAutoSlidingWindow<float, FAudioBufferAlignedAllocator> SlidingWindow(*SlidingBuffer, InAudio, HopBuffer, bDoFlush);
+
+		if (bDownmixToMono)
+		{
+			for (const FAlignedFloatBuffer& Window : SlidingWindow)
+			{
+				FMemory::Memset(MonoBuffer.GetData(), 0, sizeof(float) * NumWindowFrames);
+
+				// Sum all channels into mono buffer
+				TAutoDeinterleaveView<float, FAudioBufferAlignedAllocator> DeinterleaveView(Window, ChannelBuffer, NumChannels);
+				for(auto Channel : DeinterleaveView)
+				{
+					ArrayMixIn(Channel.Values, MonoBuffer);		
+				}
+
+				// Equal power sum. assuming incoherrent signals.
+				ArrayMultiplyByConstantInPlace(MonoBuffer, 1.f / FMath::Sqrt(static_cast<float>(NumChannels)));
+
+				AnalyzeWindow(MonoBuffer, 0, *ConstantQResult);
+
+				NumBuffers++;
+			}
+		}
+		else
+		{
+			// Loop through entire array of input samples.
+			for (const FAlignedFloatBuffer& Window : SlidingWindow)
+			{
+				// Each channel is analyzed seperately.
+				for (const TDeinterleaveView<float>::TChannel<FAudioBufferAlignedAllocator> Channel : TAutoDeinterleaveView<float, FAudioBufferAlignedAllocator>(Window, ChannelBuffer, NumChannels))
+				{
+					AnalyzeWindow(Channel.Values, Channel.ChannelIndex, *ConstantQResult);
+				}
+
+				NumBuffers++;
+			}
+		}
+	}
+
+	void FConstantQWorker::AnalyzeWindow(const FAlignedFloatBuffer& InWindow, const int32 InChannelIndex, FConstantQResult& OutResult)
+	{
+		// Run CQT Analyzer
+		ConstantQAnalyzer->CalculateCQT(InWindow.GetData(), CQTSpectrum);
+
+		// Get timestamp as center of audio window.
+		const float Timestamp = ((NumBuffers * NumHopFrames) + (0.5f * ConstantQAnalyzer->GetSettings().FFTSize)) / SampleRate;
+
+		OutResult.AddFrame(InChannelIndex, Timestamp, CQTSpectrum);
+	}
+
+	// Name of this analyzer type.
+	FName FConstantQFactory::GetName() const
+	{
+		static FName FactoryName(TEXT("ConstantQFactory"));
+		return FactoryName;
+	}
+
+	// Human readable name of this analyzer.
+	FString FConstantQFactory::GetTitle() const
+	{
+		return TEXT("Real-Time Constant Q Analyzer");
+	}
+
+	// Create a new FConstantQResult.
+	TUniquePtr<IAnalyzerResult> FConstantQFactory::NewResult() const
+	{
+		return MakeUnique<FConstantQResult>();
+	}
+
+	// Create a new FConstantQWorker
+	TUniquePtr<IAnalyzerWorker> FConstantQFactory::NewWorker(const FAnalyzerParameters& InParams, const IAnalyzerSettings* InSettings) const
+	{
+		const FConstantQSettings* ConstantQSettings = static_cast<const FConstantQSettings*>(InSettings);
+
+		check(nullptr != ConstantQSettings);
+
+		return MakeUnique<FConstantQWorker>(InParams, *ConstantQSettings);
+	}
+}
diff --git a/Engine/Plugins/Runtime/AudioSynesthesia/Source/AudioSynesthesiaCore/Public/ConstantQFactory.h b/Engine/Plugins/Runtime/AudioSynesthesia/Source/AudioSynesthesiaCore/Public/ConstantQFactory.h
new file mode 100644
--- /dev/null
+++ b/Engine/Plugins/Runtime/AudioSynesthesia/Source/AudioSynesthesiaCore/Public/ConstantQFactory.h
@@ -0,0 +1,128 @@
+// Copyright Epic Games, Inc. All Rights Reserved.
+
+#pragma once
+
+#include "IAudioAnalyzerInterface.h"
+#include "ConstantQAnalyzer.h"
+#include "ConstantQNRTFactory.h"
+
+namespace Audio
+{
+	template <typename InSampleType> class TSlidingBuffer;
+
+	/** FConstantQSettings
+	 *
+	 *  Settings for the Constant Q Non-Real-Time Factory.
+	 */
+	struct AUDIOSYNESTHESIACORE_API FConstantQSettings : public IAnalyzerSettings, public FConstantQAnalyzerSettings 
+	{
+	public:
+		/** Time, in seconds, between constant Q frames. */
+		float AnalysisPeriodInSeconds = 0.01f;
+
+		/** 
+		 * If true, all channels are mixed together with equal power before analysis. Only one channel
+		 * is produced at channel index 0.
+		 * If false, then each channel is analyzed separately. 
+		 */
+		bool bDownmixToMono = false;
+	};
+
+	/** 
+	 * FLoudnessResult contains the temporal evolution of loudness for the last buffer of audio analyzed.
+	 */
+	class AUDIOSYNESTHESIACORE_API FConstantQResult : public IAnalyzerResult
+	{
+	public:
+
+		FConstantQResult() {}
+
+		/** Add a single frame of CQT data */
+		void AddFrame(const int32 InChannelIndex, const float InTimestamp, TArrayView<const float> InSpectrum);
+
+		/** Returns const reference to FLoudnessDatum array for individual channel. */
+		const TArray<FConstantQFrame>& GetFramesForChannel(const int32 ChannelIdx) const;
+
+		/** Returns the number of channels. */
+		int32 GetNumChannels() const;
+
+	private:
+		TMap<int32, TArray<FConstantQFrame> > ChannelCQTFrames;
+	};
+
+	/** FConstantQWorker
+	 *
+	 *  FConstantQWorker computes a FConstantQResult from audio samples.
+	 */
+	class AUDIOSYNESTHESIACORE_API FConstantQWorker : public IAnalyzerWorker
+	{
+	public:
+
+		/** Constructor
+		 *
+		 * InParams are the parameters which describe characteristics of the input audio.
+		 * InAnalyzerSettings are the settings which control various aspects of the algorithm.
+		 */
+		FConstantQWorker(const FAnalyzerParameters& InParams, const FConstantQSettings& InAnalyzerSettings);
+
+		/**
+		 *  Analyze audio and put results into results pointer.
+		 *
+		 *  InAudio is an array view of audio.
+		 *  OutResult is a pointer to a valid FConstantQResult
+		 */
+		virtual void Analyze(TArrayView<const float> InAudio, IAnalyzerResult* OutResult) override;
+
+	private:
+		/** Analyze audio with multiple channels interleaved. */
+		void AnalyzeMultichannel(TArrayView<const float> InAudio, IAnalyzerResult* OutResult, const bool bDoFlush);
+
+		/** Analyze a single window of audio from a single channel */
+		void AnalyzeWindow(const FAlignedFloatBuffer& InWindow, const int32 InChannelIndex, FConstantQResult& OutResult);
+
+		int32 NumChannels = 0;
+		int32 NumBuffers = 0;
+		float SampleRate = 0.0f;
+		int32 NumHopFrames = 0;
+		int32 NumHopSamples = 0;
+		int32 NumWindowFrames = 0;
+		int32 NumWindowSamples = 0;
+
+		float MonoScaling = 1.0f;
+
+		FAlignedFloatBuffer HopBuffer;
+		FAlignedFloatBuffer ChannelBuffer;
+		FAlignedFloatBuffer MonoBuffer;
+
+		TArray<float> CQTSpectrum;
+
+		TUniquePtr<TSlidingBuffer<float> > SlidingBuffer;
+		TUniquePtr<FConstantQAnalyzer> ConstantQAnalyzer;
+
+		bool bDownmixToMono = false;
+	};
+
+	/** FConstantQFactory
+	 *  
+	 *  Factory for creating FConstantQ workers and results
+	 */
+	class AUDIOSYNESTHESIACORE_API FConstantQFactory : public IAnalyzerFactory
+	{
+	public:
+
+		/** Name of this analyzer type. */
+		virtual FName GetName() const override;
+
+		/** Human readable name of this analyzer. */
+		virtual FString GetTitle() const override;
+
+		/** Create a new FConstantQResult. */
+		virtual TUniquePtr<IAnalyzerResult> NewResult() const override;
+
+		/** Create a new FConstantQWorker 
+		 *
+		 *  InSettings must be a pointer to FConstantQSetting
+		 */
+		virtual TUniquePtr<IAnalyzerWorker> NewWorker(const FAnalyzerParameters& InParams, const IAnalyzerSettings* InSettings) const override;
+	};
+}
\ No newline at end of file
