Fix Audio Synesthesia NRT editor crash when settings object is invalid

diff --git a/Engine/Plugins/Runtime/AudioSynesthesia/Source/AudioSynesthesia/Classes/ConstantQNRT.h b/Engine/Plugins/Runtime/AudioSynesthesia/Source/AudioSynesthesia/Classes/ConstantQNRT.h
--- a/Engine/Plugins/Runtime/AudioSynesthesia/Source/AudioSynesthesia/Classes/ConstantQNRT.h
+++ b/Engine/Plugins/Runtime/AudioSynesthesia/Source/AudioSynesthesia/Classes/ConstantQNRT.h
@@ -159,7 +159,9 @@ class AUDIOSYNESTHESIA_API UConstantQNRT : public UAudioSynesthesiaNRT
 		virtual FText GetAssetActionName() const override;
 
 		virtual UClass* GetSupportedClass() const override;
-#endif
+
+		virtual bool ShouldEventTriggerAnalysis(FPropertyChangedEvent& PropertyChangeEvent) override;
+#endif // WITH_EDITOR
 
 	protected:
 
diff --git a/Engine/Plugins/Runtime/AudioSynesthesia/Source/AudioSynesthesia/Classes/LoudnessNRT.h b/Engine/Plugins/Runtime/AudioSynesthesia/Source/AudioSynesthesia/Classes/LoudnessNRT.h
--- a/Engine/Plugins/Runtime/AudioSynesthesia/Source/AudioSynesthesia/Classes/LoudnessNRT.h
+++ b/Engine/Plugins/Runtime/AudioSynesthesia/Source/AudioSynesthesia/Classes/LoudnessNRT.h
@@ -118,7 +118,9 @@ class AUDIOSYNESTHESIA_API ULoudnessNRT : public UAudioSynesthesiaNRT
 		virtual FText GetAssetActionName() const override;
 
 		virtual UClass* GetSupportedClass() const override;
-#endif
+
+		virtual bool ShouldEventTriggerAnalysis(FPropertyChangedEvent& PropertyChangeEvent) override;
+#endif // WITH_EDITOR
 
 	protected:
 
diff --git a/Engine/Plugins/Runtime/AudioSynesthesia/Source/AudioSynesthesia/Classes/OnsetNRT.h b/Engine/Plugins/Runtime/AudioSynesthesia/Source/AudioSynesthesia/Classes/OnsetNRT.h
--- a/Engine/Plugins/Runtime/AudioSynesthesia/Source/AudioSynesthesia/Classes/OnsetNRT.h
+++ b/Engine/Plugins/Runtime/AudioSynesthesia/Source/AudioSynesthesia/Classes/OnsetNRT.h
@@ -81,7 +81,9 @@ class AUDIOSYNESTHESIA_API UOnsetNRT : public UAudioSynesthesiaNRT
 		virtual FText GetAssetActionName() const override;
 
 		virtual UClass* GetSupportedClass() const override;
-#endif
+
+		virtual bool ShouldEventTriggerAnalysis(FPropertyChangedEvent& PropertyChangeEvent) override;
+#endif // WITH_EDITOR
 	protected:
 
 		/** Return the name of the IAudioAnalyzerNRTFactory associated with this UAudioAnalyzerNRT */
diff --git a/Engine/Plugins/Runtime/AudioSynesthesia/Source/AudioSynesthesia/Private/ConstantQNRT.cpp b/Engine/Plugins/Runtime/AudioSynesthesia/Source/AudioSynesthesia/Private/ConstantQNRT.cpp
--- a/Engine/Plugins/Runtime/AudioSynesthesia/Source/AudioSynesthesia/Private/ConstantQNRT.cpp
+++ b/Engine/Plugins/Runtime/AudioSynesthesia/Source/AudioSynesthesia/Private/ConstantQNRT.cpp
@@ -275,7 +275,12 @@ UClass* UConstantQNRT::GetSupportedClass() const
 {
 	return UConstantQNRT::StaticClass();
 }
-#endif
+
+bool UConstantQNRT::ShouldEventTriggerAnalysis(FPropertyChangedEvent& PropertyChangeEvent)
+{
+	return Settings != nullptr && Super::ShouldEventTriggerAnalysis(PropertyChangeEvent);
+}
+#endif // WITH_EDITOR
 
 FName UConstantQNRT::GetAnalyzerNRTFactoryName() const 
 {
diff --git a/Engine/Plugins/Runtime/AudioSynesthesia/Source/AudioSynesthesia/Private/LoudnessNRT.cpp b/Engine/Plugins/Runtime/AudioSynesthesia/Source/AudioSynesthesia/Private/LoudnessNRT.cpp
--- a/Engine/Plugins/Runtime/AudioSynesthesia/Source/AudioSynesthesia/Private/LoudnessNRT.cpp
+++ b/Engine/Plugins/Runtime/AudioSynesthesia/Source/AudioSynesthesia/Private/LoudnessNRT.cpp
@@ -186,7 +186,12 @@ UClass* ULoudnessNRT::GetSupportedClass() const
 {
 	return ULoudnessNRT::StaticClass();
 }
-#endif
+
+bool ULoudnessNRT::ShouldEventTriggerAnalysis(FPropertyChangedEvent& PropertyChangeEvent)
+{
+	return Settings != nullptr && Super::ShouldEventTriggerAnalysis(PropertyChangeEvent);
+}
+#endif // WITH_EDITOR
 
 FName ULoudnessNRT::GetAnalyzerNRTFactoryName() const
 {
diff --git a/Engine/Plugins/Runtime/AudioSynesthesia/Source/AudioSynesthesia/Private/OnsetNRT.cpp b/Engine/Plugins/Runtime/AudioSynesthesia/Source/AudioSynesthesia/Private/OnsetNRT.cpp
--- a/Engine/Plugins/Runtime/AudioSynesthesia/Source/AudioSynesthesia/Private/OnsetNRT.cpp
+++ b/Engine/Plugins/Runtime/AudioSynesthesia/Source/AudioSynesthesia/Private/OnsetNRT.cpp
@@ -219,7 +219,12 @@ UClass* UOnsetNRT::GetSupportedClass() const
 {
 	return UOnsetNRT::StaticClass();
 }
-#endif
+
+bool UOnsetNRT::ShouldEventTriggerAnalysis(FPropertyChangedEvent& PropertyChangeEvent)
+{
+	return Settings != nullptr && Super::ShouldEventTriggerAnalysis(PropertyChangeEvent);
+}
+#endif // WITH_EDITOR
 
 FName UOnsetNRT::GetAnalyzerNRTFactoryName() const
 {
