INFINITY macro usage changed to std::numeric_limits<float>::infinity() to address warnings in recent Windows SDKs (error C4756: overflow in constant arithmetic)

diff --git a/Engine/Plugins/Animation/LiveLink/Source/LiveLink/Private/LiveLinkClient.cpp b/Engine/Plugins/Animation/LiveLink/Source/LiveLink/Private/LiveLinkClient.cpp
--- a/Engine/Plugins/Animation/LiveLink/Source/LiveLink/Private/LiveLinkClient.cpp
+++ b/Engine/Plugins/Animation/LiveLink/Source/LiveLink/Private/LiveLinkClient.cpp
@@ -1901,7 +1901,7 @@ void FLiveLinkClient_Base_DEPRECATED::PushSubjectData(FGuid InSourceGuid, FName
 	}
 	for (int32 i = MaxNumberOfProperties; i < NumberOfPropertyNames; ++i)
 	{
-		NewData.PropertyValues[i] = INFINITY;
+		NewData.PropertyValues[i] = std::numeric_limits<float>::infinity();
 	}
 	PushSubjectFrameData_AnyThread(SubjectKey, MoveTemp(AnimationStruct));
 }
diff --git a/Engine/Plugins/Animation/PoseSearch/Source/Runtime/Private/PoseSearchLibrary.cpp b/Engine/Plugins/Animation/PoseSearch/Source/Runtime/Private/PoseSearchLibrary.cpp
--- a/Engine/Plugins/Animation/PoseSearch/Source/Runtime/Private/PoseSearchLibrary.cpp
+++ b/Engine/Plugins/Animation/PoseSearch/Source/Runtime/Private/PoseSearchLibrary.cpp
@@ -188,7 +188,7 @@ void FMotionMatchingState::Reset(const FTransform& ComponentTransform)
 {
 	CurrentSearchResult.Reset();
 	// Set the elapsed time to INFINITY to trigger a search right away
-	ElapsedPoseSearchTime = INFINITY;
+	ElapsedPoseSearchTime = std::numeric_limits<float>::infinity();
 	WantedPlayRate = 1.f;
 	bJumpedToPose = false;
 	ComponentDeltaYaw = 0.f;
diff --git a/Engine/Plugins/Runtime/AudioSynesthesia/Source/AudioSynesthesiaCore/Private/PeakPicker.cpp b/Engine/Plugins/Runtime/AudioSynesthesia/Source/AudioSynesthesiaCore/Private/PeakPicker.cpp
--- a/Engine/Plugins/Runtime/AudioSynesthesia/Source/AudioSynesthesiaCore/Private/PeakPicker.cpp
+++ b/Engine/Plugins/Runtime/AudioSynesthesia/Source/AudioSynesthesiaCore/Private/PeakPicker.cpp
@@ -14,7 +14,7 @@ namespace Audio
 		Settings.NumPreMean = FMath::Max(1, Settings.NumPreMean);
 		Settings.NumPostMean = FMath::Max(1, Settings.NumPostMean);
 		Settings.NumWait = FMath::Max(1, Settings.NumWait);
-		Settings.MeanDelta = FMath::Clamp(Settings.MeanDelta, 1e-6f, INFINITY);
+		Settings.MeanDelta = FMath::Clamp(Settings.MeanDelta, 1e-6f, std::numeric_limits<float>::infinity());
 	}
 
 	void FPeakPicker::PickPeaks(TArrayView<const float> InData, TArray<int32>& OutPeakIndices)
diff --git a/Engine/Source/Runtime/RenderCore/Private/RenderGraphPrivate.cpp b/Engine/Source/Runtime/RenderCore/Private/RenderGraphPrivate.cpp
--- a/Engine/Source/Runtime/RenderCore/Private/RenderGraphPrivate.cpp
+++ b/Engine/Source/Runtime/RenderCore/Private/RenderGraphPrivate.cpp
@@ -196,7 +196,7 @@ static float GetClobberValue()
 	case 2:
 		return NAN;
 	case 3:
-		return INFINITY;
+		return std::numeric_limits<float>::infinity();
 	}
 	return 0.0f;
 }
