Extend LoudnessAnalyzer to include a K-weighted curve and add ExposeOnSpawn meta specifier to Settings property

diff --git a/Engine/Plugins/Runtime/AudioSynesthesia/Source/AudioSynesthesia/Classes/Loudness.h b/Engine/Plugins/Runtime/AudioSynesthesia/Source/AudioSynesthesia/Classes/Loudness.h
--- a/Engine/Plugins/Runtime/AudioSynesthesia/Source/AudioSynesthesia/Classes/Loudness.h
+++ b/Engine/Plugins/Runtime/AudioSynesthesia/Source/AudioSynesthesia/Classes/Loudness.h
@@ -20,6 +20,8 @@ enum class ELoudnessCurveTypeEnum : uint8
 	C		UMETA(DisplayName = "C"),
 	/** Loudness Curve D Weighting. Similar to "B" but with an emphasis on presence in the 2kHz-6KHz frequency range. */
 	D		UMETA(DisplayName = "D"),
+	/** Loudness Curve K Weighting. Correlates with subjective perception of recorded broadcast content. */
+	K		UMETA(DisplayName = "K"),
 	/** No loudness curve weighting. */
 	None	UMETA(DisplayName = "None"),
 };
@@ -123,7 +125,7 @@ public:
 	ULoudnessAnalyzer();
 
 	/** The settings for the audio analyzer.  */
-	UPROPERTY(EditAnywhere, BlueprintReadWrite, Category = AudioAnalyzer)
+	UPROPERTY(EditAnywhere, BlueprintReadWrite, Category = AudioAnalyzer, meta = (ExposeOnSpawn))
 	TObjectPtr<ULoudnessSettings> Settings;
 
 	/** Delegate to receive all overall loudness results since last delegate call. */
diff --git a/Engine/Plugins/Runtime/AudioSynesthesia/Source/AudioSynesthesia/Private/Loudness.cpp b/Engine/Plugins/Runtime/AudioSynesthesia/Source/AudioSynesthesia/Private/Loudness.cpp
--- a/Engine/Plugins/Runtime/AudioSynesthesia/Source/AudioSynesthesia/Private/Loudness.cpp
+++ b/Engine/Plugins/Runtime/AudioSynesthesia/Source/AudioSynesthesia/Private/Loudness.cpp
@@ -31,6 +31,10 @@ TUniquePtr<Audio::IAnalyzerSettings> ULoudnessSettings::GetSettings(const int32
 		Settings->LoudnessCurveType = Audio::ELoudnessCurveType::D;
 		break;
 
+	case ELoudnessCurveTypeEnum::K:
+		Settings->LoudnessCurveType = Audio::ELoudnessCurveType::K;
+		break;
+
 	case ELoudnessCurveTypeEnum::None:
 	default:
 		Settings->LoudnessCurveType = Audio::ELoudnessCurveType::None;
@@ -54,7 +58,10 @@ UClass* ULoudnessSettings::GetSupportedClass() const
 
 ULoudnessAnalyzer::ULoudnessAnalyzer()
 {
-	Settings = CreateDefaultSubobject<ULoudnessSettings>(TEXT("DefaultLoudnessSettings"));
+	if (!Settings)
+	{
+		Settings = CreateDefaultSubobject<ULoudnessSettings>(TEXT("DefaultLoudnessSettings"));
+	}
 }
 
 TUniquePtr<Audio::IAnalyzerSettings> ULoudnessAnalyzer::GetSettings(const int32 InSampleRate, const int32 InNumChannels) const
diff --git a/Engine/Plugins/Runtime/AudioSynesthesia/Source/AudioSynesthesiaCore/Private/LoudnessAnalyzer.cpp b/Engine/Plugins/Runtime/AudioSynesthesia/Source/AudioSynesthesiaCore/Private/LoudnessAnalyzer.cpp
--- a/Engine/Plugins/Runtime/AudioSynesthesia/Source/AudioSynesthesiaCore/Private/LoudnessAnalyzer.cpp
+++ b/Engine/Plugins/Runtime/AudioSynesthesia/Source/AudioSynesthesiaCore/Private/LoudnessAnalyzer.cpp
@@ -45,6 +45,67 @@ namespace Audio
 		return (Freq / 6.8966888496476e-5f) * FMath::Sqrt(H / ((FreqSqrd + 79919.29f) * (FreqSqrd + 1345600.f)));
     }
 
+	float CalculatePolynomial(const float Freq, const float x0, const float x1, const float x2)
+    {
+	    const double CosFreq = FMath::Cos(Freq);
+    	return x1 * x1 + (x0 - x2) * (x0 - x2) + 2.f * x1 * (x0 + x2) * CosFreq + 4.f * x0 * x2 * CosFreq * CosFreq;
+    }
+
+	float GetEqualLoudnessKWeightForFrequency(const float Freq)
+    {
+	    constexpr float OriginalSampleRateOfKFilter = 48000.f;
+    	constexpr float Nyquist = OriginalSampleRateOfKFilter / 2.f;
+
+    	if (Freq > Nyquist)
+    	{
+    		/**
+    		 * 24kHz is not audible and does not contribute to perceived loudness.
+    		 * This may deviate from other implementations because ITU-R BS.1770-5
+    		 * does not prescribe how to handle higher sampling rates.
+    		 */
+    		return 0.f;
+    	}
+    	if (Freq < 0.f)
+    	{
+    		return 0.f;
+    	}
+
+    	const float NormalizedFrequency = Freq / OriginalSampleRateOfKFilter * 2.f * PI;
+
+    	float PreFilterPower = 0.f;
+	    {
+	    	constexpr float a0 =  1.53512485958697f;
+	    	constexpr float a1 = -2.69169618940638f;
+	    	constexpr float a2 =  1.19839281085285f;
+	    	constexpr float b0 =  1.f;
+	    	constexpr float b1 = -1.69065929318241f;
+	    	constexpr float b2 =  0.73248077421585f;
+
+	    	const float Numerator = CalculatePolynomial(NormalizedFrequency, a0, a1, a2);
+	    	const float Denominator = CalculatePolynomial(NormalizedFrequency, b0, b1, b2);
+
+	    	PreFilterPower = Denominator != 0.f ? Numerator / Denominator : 0.f;
+	    }
+
+    	float RLBFilterPower = 0.f;
+	    {
+	    	constexpr float a0 =  1.f;
+	    	constexpr float a1 = -2.f;
+	    	constexpr float a2 =  1.f;
+	    	constexpr float b0 =  1.f;
+	    	constexpr float b1 = -1.99004745483398f;
+	    	constexpr float b2 =  0.99007225036621f;
+
+	    	const float Numerator = CalculatePolynomial(NormalizedFrequency, a0, a1, a2);
+	    	const float Denominator = CalculatePolynomial(NormalizedFrequency, b0, b1, b2);
+
+	    	RLBFilterPower = Denominator != 0.f ? Numerator / Denominator : 0.f;
+	    }
+
+    	// Convert power to magnitude
+    	return FMath::Sqrt(FMath::Abs(PreFilterPower * RLBFilterPower));
+    }
+
     float GetEqualLoudnessNoneWeightForFrequency(const float Freq)
     {
         return 1.0f;
@@ -111,6 +172,10 @@ namespace Audio
 					LoudnessMapFn = GetEqualLoudnessDWeightForFrequency;
 					break;
 
+				case ELoudnessCurveType::K:
+					LoudnessMapFn = GetEqualLoudnessKWeightForFrequency;
+					break;
+
 				case ELoudnessCurveType::None:
 				default:
 					LoudnessMapFn = GetEqualLoudnessNoneWeightForFrequency;
@@ -149,6 +214,10 @@ namespace Audio
 						Weighting *= GetEqualLoudnessDWeightForFrequency(Freq);
 						break;
 
+					case ELoudnessCurveType::K:
+						Weighting *= GetEqualLoudnessKWeightForFrequency(Freq);
+						break;
+
 					case ELoudnessCurveType::None:
 					default:
 						Weighting *= GetEqualLoudnessNoneWeightForFrequency(Freq);
diff --git a/Engine/Plugins/Runtime/AudioSynesthesia/Source/AudioSynesthesiaCore/Public/LoudnessAnalyzer.h b/Engine/Plugins/Runtime/AudioSynesthesia/Source/AudioSynesthesiaCore/Public/LoudnessAnalyzer.h
--- a/Engine/Plugins/Runtime/AudioSynesthesia/Source/AudioSynesthesiaCore/Public/LoudnessAnalyzer.h
+++ b/Engine/Plugins/Runtime/AudioSynesthesia/Source/AudioSynesthesiaCore/Public/LoudnessAnalyzer.h
@@ -15,6 +15,7 @@ namespace Audio
 		B,
 		C,
 		D,
+		K,
 		None
 	};
 
