Add k-weighted loudness analysis to NRT analyzer.
[RN] Added k-weighted analyzes to non-real-time analyzer

diff --git a/Engine/Plugins/Runtime/AudioSynesthesia/Source/AudioSynesthesia/Classes/Loudness.h b/Engine/Plugins/Runtime/AudioSynesthesia/Source/AudioSynesthesia/Classes/Loudness.h
--- a/Engine/Plugins/Runtime/AudioSynesthesia/Source/AudioSynesthesia/Classes/Loudness.h
+++ b/Engine/Plugins/Runtime/AudioSynesthesia/Source/AudioSynesthesia/Classes/Loudness.h
@@ -129,19 +129,19 @@ public:
 	TObjectPtr<ULoudnessSettings> Settings;
 
 	/** Delegate to receive all overall loudness results since last delegate call. */
-	UPROPERTY(BlueprintAssignable)
+	UPROPERTY(BlueprintAssignable, Category = AudioAnalyzer)
 	FOnOverallLoudnessResults OnOverallLoudnessResults;
 
 	/** Delegate to receive all loudness results, per-channel, since last delegate call. */
-	UPROPERTY(BlueprintAssignable)
+	UPROPERTY(BlueprintAssignable, Category = AudioAnalyzer)
 	FOnPerChannelLoudnessResults OnPerChannelLoudnessResults;
 
 	/** Delegate to receive the latest overall loudness results. */
-	UPROPERTY(BlueprintAssignable)
+	UPROPERTY(BlueprintAssignable, Category = AudioAnalyzer)
 	FOnLatestOverallLoudnessResults OnLatestOverallLoudnessResults;
 
 	/** Delegate to receive the latest per-channel loudness results. */
-	UPROPERTY(BlueprintAssignable)
+	UPROPERTY(BlueprintAssignable, Category = AudioAnalyzer)
 	FOnLatestPerChannelLoudnessResults OnLatestPerChannelLoudnessResults;
 
 	/** Convert ULoudnessSettings to FLoudnessSettings */
@@ -154,4 +154,4 @@ protected:
 
 	/** Return the name of the IAudioAnalyzerFactory associated with this UAudioAnalyzer */
 	FName GetAnalyzerFactoryName() const override;
-};
\ No newline at end of file
+};
diff --git a/Engine/Plugins/Runtime/AudioSynesthesia/Source/AudioSynesthesia/Classes/LoudnessNRT.h b/Engine/Plugins/Runtime/AudioSynesthesia/Source/AudioSynesthesia/Classes/LoudnessNRT.h
--- a/Engine/Plugins/Runtime/AudioSynesthesia/Source/AudioSynesthesia/Classes/LoudnessNRT.h
+++ b/Engine/Plugins/Runtime/AudioSynesthesia/Source/AudioSynesthesia/Classes/LoudnessNRT.h
@@ -20,6 +20,8 @@ enum class ELoudnessNRTCurveTypeEnum : uint8
 	C		UMETA(DisplayName="C"),
 	/** Loudness Curve D Weighting. Similar to "B" but with an emphasis on presence in the 2kHz-6KHz frequency range. */
 	D		UMETA(DisplayName="D"),
+	/** Loudness Curve K Weighting. Correlates with subjective perception of recorded broadcast content. */
+	K		UMETA(DisplayName = "K"),
 	/** No loudness curve weighting. */
 	None	UMETA(DisplayName="None"),
 };
@@ -84,7 +86,7 @@ class AUDIOSYNESTHESIA_API ULoudnessNRT : public UAudioSynesthesiaNRT
 		ULoudnessNRT();
 
 		/** The settings for the audio analyzer.  */
-		UPROPERTY(EditAnywhere, BlueprintReadOnly, Category=AudioAnalyzer)
+		UPROPERTY(EditAnywhere, BlueprintReadOnly, Category=AudioAnalyzer, meta = (ExposeOnSpawn))
 		TObjectPtr<ULoudnessNRTSettings> Settings;
 
 		/** Get the overall loudness of the analyzed sound at a given time. */
diff --git a/Engine/Plugins/Runtime/AudioSynesthesia/Source/AudioSynesthesia/Private/LoudnessNRT.cpp b/Engine/Plugins/Runtime/AudioSynesthesia/Source/AudioSynesthesia/Private/LoudnessNRT.cpp
--- a/Engine/Plugins/Runtime/AudioSynesthesia/Source/AudioSynesthesia/Private/LoudnessNRT.cpp
+++ b/Engine/Plugins/Runtime/AudioSynesthesia/Source/AudioSynesthesia/Private/LoudnessNRT.cpp
@@ -65,6 +65,10 @@ TUniquePtr<Audio::IAnalyzerNRTSettings> ULoudnessNRTSettings::GetSettings(const
 			Settings->LoudnessCurveType = Audio::ELoudnessCurveType::D;
 			break;
 
+		case ELoudnessNRTCurveTypeEnum::K:
+			Settings->LoudnessCurveType = Audio::ELoudnessCurveType::K;
+			break;
+
 		case ELoudnessNRTCurveTypeEnum::None:
 		default:
 			Settings->LoudnessCurveType = Audio::ELoudnessCurveType::None;
