Update the LoundessAnalyzer to match LKFS/LUFS loudness standard.

diff --git a/Engine/Plugins/Runtime/AudioSynesthesia/Source/AudioSynesthesiaCore/Private/LoudnessAnalyzer.cpp b/Engine/Plugins/Runtime/AudioSynesthesia/Source/AudioSynesthesiaCore/Private/LoudnessAnalyzer.cpp
--- a/Engine/Plugins/Runtime/AudioSynesthesia/Source/AudioSynesthesiaCore/Private/LoudnessAnalyzer.cpp
+++ b/Engine/Plugins/Runtime/AudioSynesthesia/Source/AudioSynesthesiaCore/Private/LoudnessAnalyzer.cpp
@@ -118,11 +118,11 @@ namespace Audio
 
     FLoudnessAnalyzer::FLoudnessAnalyzer(float InSampleRate, const FLoudnessAnalyzerSettings& InSettings)
     :   Settings(InSettings)
-    ,   Window(Settings.WindowType, Settings.FFTSize, 1, false)
     {
 		checkf(FMath::IsPowerOfTwo(Settings.FFTSize), TEXT("FFT size must be a power of two"));
 		checkf(Settings.FFTSize >= 8, TEXT("FFT size must be atleast 8"))
         check(InSampleRate > 0);
+		
 
 		// Create FFT
 		FFFTSettings FFTSettings;
@@ -136,6 +136,17 @@ namespace Audio
 		FFT = FFFTFactory::NewFFTAlgorithm(FFTSettings);
 		if (FFT.IsValid())
 		{
+			checkf(Settings.WindowSize <= ActualFFTSize, TEXT("WindowSize %d must be less than or equal to FFTSize %d"), Settings.WindowSize, ActualFFTSize);
+
+			if (Settings.WindowSize <= 0)
+			{
+				Settings.WindowSize = ActualFFTSize;
+			}
+		
+			// Initialize the analysis window
+			Window.AddUninitialized(Settings.WindowSize);
+			GenerateWindow(Window.GetData(), Settings.WindowType, Window.Num(), 1 /* NumChannels */, false /* bIsPeriodic */);
+			
 			// Size internal buffers
 			WindowedSamples.AddZeroed(FFT->NumInputFloats()); // Zero samples to apply zero buffer in case actual FFT is larger than provided settings.
 			ComplexSpectrum.AddUninitialized(FFT->NumOutputFloats());
@@ -153,40 +164,6 @@ namespace Audio
 			CurveWeights.Reset(FFTFreqSize);
 			CurveWeights.AddUninitialized(FFTFreqSize);
 
-			float(*LoudnessMapFn)(const float);
-			switch (Settings.LoudnessCurveType)
-			{
-				case ELoudnessCurveType::A:
-					LoudnessMapFn = GetEqualLoudnessAWeightForFrequency;
-					break;
-
-				case ELoudnessCurveType::B:
-					LoudnessMapFn = GetEqualLoudnessBWeightForFrequency;
-					break;
-
-				case ELoudnessCurveType::C:
-					LoudnessMapFn = GetEqualLoudnessCWeightForFrequency;
-					break;
-
-				case ELoudnessCurveType::D:
-					LoudnessMapFn = GetEqualLoudnessDWeightForFrequency;
-					break;
-
-				case ELoudnessCurveType::K:
-					LoudnessMapFn = GetEqualLoudnessKWeightForFrequency;
-					break;
-
-				case ELoudnessCurveType::None:
-				default:
-					LoudnessMapFn = GetEqualLoudnessNoneWeightForFrequency;
-					break;
-			}
-
-			/* Calculate equal loudness weighting curve */
-
-			// Convention is to normalize 1kHz to have a loudness of 1.
-			float EqualLoudnessNorm = 1.0f / LoudnessMapFn(1000.0f);
-
 			float* CurveWeightsPtr = CurveWeights.GetData();
 
 			const float FFTSize = static_cast<float>(ActualFFTSize);
@@ -194,43 +171,13 @@ namespace Audio
 			for (int32 i = 0; i < FFTFreqSize; i++)
 			{
 				const float Freq = static_cast<float>(i) * InSampleRate / FFTSize;
-				float Weighting = EqualLoudnessNorm;
-
-				switch (Settings.LoudnessCurveType)
-				{
-					case ELoudnessCurveType::A:
-						Weighting *= GetEqualLoudnessAWeightForFrequency(Freq);
-						break;
-
-					case ELoudnessCurveType::B:
-						Weighting *= GetEqualLoudnessBWeightForFrequency(Freq);
-						break;
-
-					case ELoudnessCurveType::C:
-						Weighting *= GetEqualLoudnessCWeightForFrequency(Freq);
-						break;
-
-					case ELoudnessCurveType::D:
-						Weighting *= GetEqualLoudnessDWeightForFrequency(Freq);
-						break;
-
-					case ELoudnessCurveType::K:
-						Weighting *= GetEqualLoudnessKWeightForFrequency(Freq);
-						break;
-
-					case ELoudnessCurveType::None:
-					default:
-						Weighting *= GetEqualLoudnessNoneWeightForFrequency(Freq);
-						break;
-				}
+				const float Weighting = GetEqualLoudnessWeightForFrequency(Freq);
 
 				// Curve designed for magnitude domain, but applied to power spectrum. The curve is squared to be applied to power spectrum.
 				CurveWeightsPtr[i] = Weighting * Weighting;
 			}
 
-			// Normalize by FFT Scaling (1 / FFTSize) and by number of samples (1 / WindowSize) to calculate mean squared value.
-			EnergyScale = GetPowerSpectrumScaling(ActualFFTSize, FFT->ForwardScaling(), EFFTScaling::None);
-			EnergyScale *= 1.f / static_cast<float>(Settings.FFTSize);
+			EnergyScale = CalculateEnergyScaling();
 		}
     }
 
@@ -241,15 +188,14 @@ namespace Audio
 
     float FLoudnessAnalyzer::CalculatePerceptualEnergy(TArrayView<const float> InView)
     {
-		check(InView.Num() == Settings.FFTSize);
+		check(InView.Num() == Window.Num());
 
 		float Total = 0.f;
 
 		if (FFT.IsValid())
 		{
 			// Copy input samples and apply window
-			FMemory::Memcpy(WindowedSamples.GetData(), InView.GetData(), sizeof(float) * Settings.FFTSize);
-			Window.ApplyToBuffer(WindowedSamples.GetData());
+			ArrayMultiply(InView, Window, TArrayView<float>(WindowedSamples.GetData(), Window.Num()));
 
 			FFT->ForwardRealToComplex(WindowedSamples.GetData(), ComplexSpectrum.GetData());
 
@@ -275,11 +221,119 @@ namespace Audio
 		return ConvertPerceptualEnergyToLoudness(Energy);
     }
 
-	float FLoudnessAnalyzer::ConvertPerceptualEnergyToLoudness(float InPerceptualEnergy)
+	float FLoudnessAnalyzer::ConvertPerceptualEnergyToLoudness(float InPerceptualEnergy) 
 	{
         return 10.0f * FMath::Loge(InPerceptualEnergy) * LoudnessAnalyzerPrivate::Log10Scale;
 	}
 
+	float FLoudnessAnalyzer::ConvertPerceptualEnergyToLKFS(float InPerceptualEnergy) 
+	{
+        return (10.0f * FMath::Loge(InPerceptualEnergy) * LoudnessAnalyzerPrivate::Log10Scale) - 0.691f;
+	}
+
+	float FLoudnessAnalyzer::GetEqualLoudnessWeightForFrequency(float InHz) const
+	{
+		switch (Settings.LoudnessCurveType)
+		{
+			case ELoudnessCurveType::A:
+				return GetEqualLoudnessAWeightForFrequency(InHz);
+
+			case ELoudnessCurveType::B:
+				return GetEqualLoudnessBWeightForFrequency(InHz);
+
+			case ELoudnessCurveType::C:
+				return GetEqualLoudnessCWeightForFrequency(InHz);
+
+			case ELoudnessCurveType::D:
+				return GetEqualLoudnessDWeightForFrequency(InHz);
+
+			case ELoudnessCurveType::K:
+				return GetEqualLoudnessKWeightForFrequency(InHz);
+
+			case ELoudnessCurveType::None:
+				return GetEqualLoudnessNoneWeightForFrequency(InHz);
+
+			default:
+				{
+					checkNoEntry();
+				}
+				return GetEqualLoudnessNoneWeightForFrequency(InHz);
+		}
+	}
+
+	float FLoudnessAnalyzer::CalculateEnergyScaling() const
+	{
+		switch (Settings.ScalingMethod)
+		{
+			case ELoudnessAnalyzerScalingMethod::BackCompatible:
+				return CalculateEnergyScalingBackCompatible();
+
+			case ELoudnessAnalyzerScalingMethod::Corrected:
+				return CalculateEnergyScalingCorrected();
+
+			default:
+				{
+					checkNoEntry(); // Unimplemented normalization method
+				}
+				return CalculateEnergyScalingCorrected();
+		};
+	}
+
+	float FLoudnessAnalyzer::CalculateEnergyScalingBackCompatible() const
+	{
+		check(FFT.IsValid()); // FFT must be initialized before calling this function 
+
+		const float FFTScaling = GetPowerSpectrumScaling(FFT->Size(), FFT->ForwardScaling(), EFFTScaling::None);
+		const float SizeScaling = 1.f / static_cast<float>(Settings.FFTSize);
+		const float CurveAt1Khz = GetEqualLoudnessWeightForFrequency(1000.f);
+		float CurveScaling = 1.f;
+
+		if (CurveAt1Khz > 0.f)
+		{
+			CurveScaling = 1.f / CurveAt1Khz; // Normalize curve to be unity at 1khz
+			CurveScaling = CurveScaling * CurveScaling; // Convert amplitude normalization to power normalization
+		}
+
+		return FFTScaling * SizeScaling * CurveScaling;
+	}
+
+	float FLoudnessAnalyzer::CalculateEnergyScalingCorrected() const
+	{
+		check(FFT.IsValid()); // FFT must be initialized before calling this function 
+
+		constexpr float HalfSpectrumScaling = 2.f; // When accumulating spectral energy, we ignore the negative spectrum
+												   // magnitudes because they are symmetric with positive spectrum
+												   // magnitudes. This results in accumulation of only half of the 
+												   // spectrum. This scaling corrects that by doubling the energy
+												   // after it has been accumulated.
+		const float FFTScaling = GetPowerSpectrumScaling(FFT->Size(), FFT->ForwardScaling(), EFFTScaling::None);
+		const float WindowScaling = CalculateWindowEnergyScaling();
+
+		return HalfSpectrumScaling * FFTScaling * WindowScaling;
+	}
+
+	float FLoudnessAnalyzer::CalculateWindowEnergyScaling() const
+	{
+		check(Window.Num() > 0); // Window must be initialized before calling this function.
+
+		// Calculate energy of window
+		float WindowEnergy = 0.f;
+		for (float Val : Window)
+		{
+			WindowEnergy += (Val * Val);
+		}
+
+		// Avoid divide by zero on empty window func
+		if (WindowEnergy == 0.f)
+		{
+			return 1.f;
+		}
+		else
+		{
+			return 1.f / WindowEnergy;
+		}
+	}
+
 	/****************************************************************************/
 	/****************    FMultichannelLoudnessAnalyzer    ***********************/
 	/****************************************************************************/
@@ -294,13 +348,12 @@ namespace Audio
         ChannelWeights[EAudioMixerChannel::BackLeft] = 1.5f;
         ChannelWeights[EAudioMixerChannel::BackRight] = 1.5f;
 
-        MonoBuffer.Reset(InSettings.FFTSize);
-        MonoBuffer.AddUninitialized(InSettings.FFTSize);
+        MonoBuffer.AddUninitialized(Analyzer.GetSettings().WindowSize);
     }
 
     float FMultichannelLoudnessAnalyzer::CalculatePerceptualEnergy(TArrayView<const float> InView, const int32 InNumChannels, TArray<float>& OutChannelEnergies)
     {
-		check(InView.Num() == (InNumChannels * Analyzer.GetSettings().FFTSize));
+		check(InView.Num() == (InNumChannels * Analyzer.GetSettings().WindowSize));
 
         OutChannelEnergies.Reset(InNumChannels);
         OutChannelEnergies.AddUninitialized(InNumChannels);
@@ -337,5 +390,6 @@ namespace Audio
     {
         return Analyzer.GetSettings();
     }
+
 }
 
diff --git a/Engine/Plugins/Runtime/AudioSynesthesia/Source/AudioSynesthesiaCore/Public/LoudnessAnalyzer.h b/Engine/Plugins/Runtime/AudioSynesthesia/Source/AudioSynesthesiaCore/Public/LoudnessAnalyzer.h
--- a/Engine/Plugins/Runtime/AudioSynesthesia/Source/AudioSynesthesiaCore/Public/LoudnessAnalyzer.h
+++ b/Engine/Plugins/Runtime/AudioSynesthesia/Source/AudioSynesthesiaCore/Public/LoudnessAnalyzer.h
@@ -36,32 +36,49 @@ namespace Audio
 	/** Generate equal loudness weight for magnitude spectrum using curve None */
     AUDIOSYNESTHESIACORE_API float GetEqualLoudnessNoneWeigthForFrequency(const float Freq);
 
+	enum class ELoudnessAnalyzerScalingMethod
+	{
+		// Normalize by FFTSize and having perceptual curves be unity gain at 1khz. 
+		// 
+		// Note: This method does not correctly account for different window functions
+		// and also reduces the energy by half. It remains here for backwards
+		// compatibility.
+		BackCompatible,
+
+		// Normalize by Window Function. Use this when performing loudness analysis
+		// which must match the LKFS standardized measurements. 
+		Corrected
+	};
+
 	/** Settings for FLoudnessAnalyzer */
     struct FLoudnessAnalyzerSettings
     {
 		/** Size of FFT describe in number of samples */
-        int32 FFTSize;
+        int32 FFTSize = 8192;
 
 		/** Type of window to be applied to input audio */
-        EWindowType WindowType;
+        EWindowType WindowType = EWindowType::Blackman;
 
 		/** Type of equal loudness curve to weight power spectrum */
-        ELoudnessCurveType LoudnessCurveType;
+        ELoudnessCurveType LoudnessCurveType = ELoudnessCurveType::D;
 
 		/** Minimum frequency for analyzing audio */
-        float MinAnalysisFrequency;
+        float MinAnalysisFrequency = 20.f;
 
 		/** Maximum frequency for analyzing audio */
-        float MaxAnalysisFrequency;
-
-        /** Default settings */
-		FLoudnessAnalyzerSettings()
-			: FFTSize(8192)
-			, WindowType(EWindowType::Blackman)
-			, LoudnessCurveType(ELoudnessCurveType::D)
-			, MinAnalysisFrequency(20.0f)
-			, MaxAnalysisFrequency(10000.0f)
-		{}
+        float MaxAnalysisFrequency = 10000.f;
+
+		/** Size of the audio analysis window in seconds. 
+		 *
+		 * If this is less than 0, it will become equal to the FFTSize.
+		 * Otherwise, this must be less than or equal to the FFT Size. 
+		 */
+		int32 WindowSize = -1;
+
+		/** Method to use for scaling and normalizing data. Corrected is preferred
+		 * in terms of correctness. BackCompatible is available for backwards compatibility.
+		 */
+		ELoudnessAnalyzerScalingMethod ScalingMethod = ELoudnessAnalyzerScalingMethod::BackCompatible;
     };
 
 	/** FLoudnessAnalyzer
@@ -92,6 +109,11 @@ namespace Audio
 			 * Convert from perceptual energy to loudness [Loudness = 10 * log_10(Energy)]
 			 */
 			static UE_API float ConvertPerceptualEnergyToLoudness(float InPerceptualEnergy);
+			
+			/**
+			 * Convert from perceptual energy to loudness [LKFS = -0.691 + 10 * log_10(Energy)]
+			 */
+			static UE_API float ConvertPerceptualEnergyToLKFS(float InPerceptualEnergy);
 
 			/**
 			 * Return const reference to settings used inside this analyzer.
@@ -99,6 +121,12 @@ namespace Audio
             UE_API const FLoudnessAnalyzerSettings& GetSettings() const;
         private:
 
+			float CalculateEnergyScaling() const;
+			float GetEqualLoudnessWeightForFrequency(float InHz) const;
+			float CalculateEnergyScalingBackCompatible() const;
+			float CalculateEnergyScalingCorrected() const;
+			float CalculateWindowEnergyScaling() const;
+
             FLoudnessAnalyzerSettings Settings;
 
 			/** Scale factor used to normalize between differing window sizes and fft sizes */
@@ -111,9 +139,9 @@ namespace Audio
 			int32 FFTFreqSize = 0;
 			int32 ActualFFTSize = 0;
 
-            FWindow Window;
 			TUniquePtr<IFFTAlgorithm> FFT;
 
+			FAlignedFloatBuffer Window;
             FAlignedFloatBuffer WindowedSamples;
             FAlignedFloatBuffer ComplexSpectrum;
             FAlignedFloatBuffer RealSpectrum;
diff --git a/Engine/Source/Runtime/SignalProcessing/Private/AudioFFT.cpp b/Engine/Source/Runtime/SignalProcessing/Private/AudioFFT.cpp
--- a/Engine/Source/Runtime/SignalProcessing/Private/AudioFFT.cpp
+++ b/Engine/Source/Runtime/SignalProcessing/Private/AudioFFT.cpp
@@ -18,6 +18,15 @@ TAutoConsoleVariable<int32> CVarFFTMethod(
 
 namespace Audio
 {
+	void GenerateRectangularWindow(float* WindowBuffer, int32 NumFrames, int32 NumChannels)
+	{
+		int32 NumSamples = NumFrames * NumChannels;
+		check(NumSamples >= 0);
+		for (int32 i = 0; i < NumSamples; i++)
+		{
+			WindowBuffer[i] = 1.f;
+		}
+	}
 
 	void GenerateHammingWindow(float* WindowBuffer, int32 NumFrames, int32 NumChannels, bool bIsPeriodic)
 	{
@@ -86,6 +95,34 @@ namespace Audio
 		}
 	}
 
+	void GenerateWindow(float* InOutWindowBuffer, EWindowType InType, int32 InNumFrames, int32 InNumChannels, bool bIsPeriodic)
+	{
+		switch (InType)
+		{
+		case EWindowType::None:
+			GenerateRectangularWindow(InOutWindowBuffer, InNumFrames, InNumChannels);	
+			break;
+
+		case EWindowType::Hamming:
+			GenerateHammingWindow(InOutWindowBuffer, InNumFrames, InNumChannels, bIsPeriodic);	
+			break;
+
+		case EWindowType::Hann:
+			GenerateHannWindow(InOutWindowBuffer, InNumFrames, InNumChannels, bIsPeriodic);	
+			break;
+
+		case EWindowType::Blackman:
+			GenerateBlackmanWindow(InOutWindowBuffer, InNumFrames, InNumChannels, bIsPeriodic);	
+			break;
+
+		default:
+			{
+				checkNoEntry();
+				GenerateRectangularWindow(InOutWindowBuffer, InNumFrames, InNumChannels);	
+			}
+		}
+	}
+
 	uint32 GetCOLAHopSizeForWindow(EWindowType InType, uint32 WindowLength)
 	{
 		switch (InType)
diff --git a/Engine/Source/Runtime/SignalProcessing/Public/DSP/AudioFFT.h b/Engine/Source/Runtime/SignalProcessing/Public/DSP/AudioFFT.h
--- a/Engine/Source/Runtime/SignalProcessing/Public/DSP/AudioFFT.h
+++ b/Engine/Source/Runtime/SignalProcessing/Public/DSP/AudioFFT.h
@@ -28,9 +28,11 @@ namespace Audio
 	};
 
 	// Utility functions for generating different types of windows. Called in FWindow::Generate.
+	SIGNALPROCESSING_API void GenerateRectangularWindow(float* WindowBuffer, int32 NumFrames, int32 NumChannels);
 	SIGNALPROCESSING_API void GenerateHammingWindow(float* WindowBuffer, int32 NumFrames, int32 NumChannels, bool bIsPeriodic);
 	SIGNALPROCESSING_API void GenerateHannWindow(float* WindowBuffer, int32 NumFrames, int32 NumChannels, bool bIsPeriodic);
 	SIGNALPROCESSING_API void GenerateBlackmanWindow(float* WindowBuffer, int32 NumFrames, int32 NumChannels, bool bIsPeriodic);
+	SIGNALPROCESSING_API void GenerateWindow(float* InOutWindowBuffer, EWindowType InType, int32 InNumFrames, int32 InNumChannels, bool bIsPeriodic);
 
 	// Returns the hop size in samples necessary to maintain constant overlap add.
 	// For more information on COLA, see the following page:
@@ -71,36 +73,6 @@ namespace Audio
 		SIGNALPROCESSING_API void Generate(int32 NumFrames, int32 NumChannels, bool bIsPeriodic);
 	};
 
-	
-	struct UE_DEPRECATED(5.5, "FrequencyBuffer is no longer used.") FrequencyBuffer
-	{
-		FAlignedFloatBuffer Real;
-		FAlignedFloatBuffer Imag;
-
-		void InitZeroed(int32 Num)
-		{
-			Real.Reset();
-			Real.AddZeroed(Num);
-
-			Imag.Reset();
-			Imag.AddZeroed(Num);
-		}
-
-		void CopyFrom(const float* InReal, const float* InImag, int32 Num)
-		{
-			check(Num == Real.Num() && Num == Imag.Num());
-			FMemory::Memcpy(Real.GetData(), InReal, Num * sizeof(float));
-			FMemory::Memcpy(Imag.GetData(), InImag, Num * sizeof(float));
-		}
-
-		void CopyFrom(const FrequencyBuffer& Other)
-		{
-			check(Other.Real.Num() == Real.Num() && Other.Imag.Num() == Imag.Num());
-			FMemory::Memcpy(Real.GetData(), Other.Real.GetData(), Other.Real.Num() * sizeof(float));
-			FMemory::Memcpy(Imag.GetData(), Other.Imag.GetData(), Other.Imag.Num() * sizeof(float));
-		}
-	};
-
 	// Return the ceiling of the log2 of InNum
 	SIGNALPROCESSING_API int32 CeilLog2(int32 InNum);
 
