LoudnessAnalyzer linter fixes
#rnx

diff --git a/Engine/Plugins/Runtime/AudioSynesthesia/Source/AudioSynesthesiaCore/Public/LoudnessAnalyzer.h b/Engine/Plugins/Runtime/AudioSynesthesia/Source/AudioSynesthesiaCore/Public/LoudnessAnalyzer.h
--- a/Engine/Plugins/Runtime/AudioSynesthesia/Source/AudioSynesthesiaCore/Public/LoudnessAnalyzer.h
+++ b/Engine/Plugins/Runtime/AudioSynesthesia/Source/AudioSynesthesiaCore/Public/LoudnessAnalyzer.h
@@ -36,7 +36,7 @@ namespace Audio
 	/** Generate equal loudness weight for magnitude spectrum using curve None */
     AUDIOSYNESTHESIACORE_API float GetEqualLoudnessNoneWeigthForFrequency(const float Freq);
 
-	enum class ELoudnessAnalyzerScalingMethod
+	enum class ELoudnessAnalyzerScalingMethod : uint8
 	{
 		// Normalize by FFTSize and having perceptual curves be unity gain at 1khz. 
 		// 
@@ -68,10 +68,10 @@ namespace Audio
 		/** Maximum frequency for analyzing audio */
         float MaxAnalysisFrequency = 10000.f;
 
-		/** Size of the audio analysis window in seconds. 
+		/** Size of the audio analysis window in samples.
 		 *
-		 * If this is less than 0, it will become equal to the FFTSize.
-		 * Otherwise, this must be less than or equal to the FFT Size. 
+		 * If this is less than 1, it will become equal to the FFTSize.
+		 * Otherwise, this must be less than or equal to the FFT Size.
 		 */
 		int32 WindowSize = -1;
 
