Adding LKFS NRT analyzer

diff --git a/Engine/Plugins/Runtime/AudioSynesthesia/Source/AudioSynesthesia/Classes/LKFSNRT.h b/Engine/Plugins/Runtime/AudioSynesthesia/Source/AudioSynesthesia/Classes/LKFSNRT.h
new file mode 100644
--- /dev/null
+++ b/Engine/Plugins/Runtime/AudioSynesthesia/Source/AudioSynesthesia/Classes/LKFSNRT.h
@@ -0,0 +1,124 @@
+// Copyright Epic Games, Inc. All Rights Reserved.
+
+#pragma once
+
+#include "AudioSynesthesiaNRT.h"
+#include "LKFSNRTFactory.h"
+#include "LKFSNRT.generated.h"
+
+#define UE_API AUDIOSYNESTHESIA_API
+
+
+/** ULKFSNRTSettings
+ *
+ * Settings for a ULKFSNRT analyzer.
+ */
+UCLASS(MinimalAPI, Blueprintable)
+class ULKFSNRTSettings : public UAudioSynesthesiaNRTSettings
+{
+	GENERATED_BODY()
+	public:
+
+		UE_API ULKFSNRTSettings();
+
+		/** Number of seconds between loudness measurements */
+		UPROPERTY(EditAnywhere, BlueprintReadOnly, Category=AudioAnalyzer, meta = (ClampMin = "0.01", ClampMax = "1.00"))
+		float AnalysisPeriod = 0.1f;
+
+		/** Number of seconds of audio analyzed for a single window of measurement*/
+		UPROPERTY(EditAnywhere, BlueprintReadOnly, Category=AudioAnalyzer, meta = (ClampMin = "0.10", ClampMax = "1.00"))
+		float AnalysisWindowDuration = 0.4f;
+
+		/** Number of seconds of audio analyzed when determining short term loudness.*/
+		UPROPERTY(EditAnywhere, BlueprintReadOnly, Category=AudioAnalyzer, meta = (ClampMin = "1.00", ClampMax = "30.00"))
+		float ShortTermLoudnessDuration = Audio::FLKFSNRTSettings::DefaultShortTermLoudnessDuration;
+
+		/** Convert ULKFSNRTSettings to FLKFSNRTSettings */
+		UE_API virtual TUniquePtr<Audio::IAnalyzerNRTSettings> GetSettings(const float InSampleRate, const int32 InNumChannels) const;
+
+#if WITH_EDITOR
+		UE_API virtual FText GetAssetActionName() const override;
+
+		UE_API virtual UClass* GetSupportedClass() const override;
+#endif
+};
+
+
+
+/** ULKFSNRT
+ *
+ * ULKFSNRT calculates the temporal evolution of perceptual loudness for a given 
+ * sound. 
+ */
+UCLASS(MinimalAPI, Blueprintable)
+class ULKFSNRT : public UAudioSynesthesiaNRT
+{
+	GENERATED_BODY()
+
+	public:
+
+		UE_API ULKFSNRT();
+
+		/** The settings for the audio analyzer.  */
+		UPROPERTY(EditAnywhere, BlueprintReadOnly, Category=AudioAnalyzer, meta = (ExposeOnSpawn))
+		TObjectPtr<ULKFSNRTSettings> Settings;
+
+		/** Get the overall loudness of the analyzed sound at a given time. */
+		UFUNCTION(BlueprintCallable, Category="Audio Analyzer")
+		UE_API void GetLoudnessAtTime(const float InSeconds, float& OutLoudness) const;
+
+		/** Get a specific channel loudness of the analyzed sound at a given time. */
+		UFUNCTION(BlueprintCallable, Category="Audio Analyzer")
+		UE_API void GetChannelLoudnessAtTime(const float InSeconds, const int32 InChannel, float& OutLoudness) const;
+
+		/** Return the entire array of loudness data for a specific audio channel. */
+		UFUNCTION(BlueprintCallable, Category="Audio Analyzer")
+		UE_API TArray<FLKFSNRTResults> GetLoudnessDataForChannel(int32 InChannel) const;
+
+		/** Return the entire array of loudness data for the combination of all audio channels. */
+		UFUNCTION(BlueprintCallable, Category="Audio Analyzer")
+		UE_API TArray<FLKFSNRTResults> GetLoudnessData() const;
+
+		/** Return the instantaneous loudness data for a specific time and audio channel. */
+		UFUNCTION(BlueprintCallable, Category="Audio Analyzer")
+		UE_API FLKFSNRTResults GetLoudnessDataForChannelAtTime(const float InSeconds, int32 InChannel) const;
+
+		/** Return the instantaneous loudness data for the combination of all audio channels at a specific time. */
+		UFUNCTION(BlueprintCallable, Category="Audio Analyzer")
+		UE_API FLKFSNRTResults GetLoudnessDataAtTime(const float InSeconds) const;
+
+		/** Return the integrated (average) loudness for a specific audio channel. */
+		UFUNCTION(BlueprintCallable, Category="Audio Analyzer")
+		UE_API float GetIntegratedLoudnessForChannel(int32 InChannel) const;
+
+		/** Return the integrated (average) loudness for the combination of all audio channels. */
+		UFUNCTION(BlueprintCallable, Category="Audio Analyzer")
+		UE_API float GetIntegratedLoudness() const;
+
+		/** Return the gated loudness for a specific audio channel. */
+		UFUNCTION(BlueprintCallable, Category="Audio Analyzer")
+		UE_API float GetGatedLoudnessForChannel(int32 InChannel) const;
+
+		/** Return the gated loudness for the combination of all audio channels. */
+		UFUNCTION(BlueprintCallable, Category="Audio Analyzer")
+		UE_API float GetGatedLoudness() const;
+
+		/** Convert ULKFSNRTSettings to FLKFSNRTSettings */
+ 		UE_API TUniquePtr<Audio::IAnalyzerNRTSettings> GetSettings(const float InSampleRate, const int32 InNumChannels) const override;
+
+#if WITH_EDITOR
+		UE_API virtual FText GetAssetActionName() const override;
+
+		UE_API virtual UClass* GetSupportedClass() const override;
+
+		UE_API virtual bool ShouldEventTriggerAnalysis(FPropertyChangedEvent& PropertyChangeEvent) override;
+#endif // WITH_EDITOR
+
+	protected:
+
+		/** Return the name of the IAudioAnalyzerNRTFactory associated with this UAudioAnalyzerNRT */
+		UE_API FName GetAnalyzerNRTFactoryName() const override;
+
+};
+
+#undef UE_API
diff --git a/Engine/Plugins/Runtime/AudioSynesthesia/Source/AudioSynesthesia/Private/LKFSNRT.cpp b/Engine/Plugins/Runtime/AudioSynesthesia/Source/AudioSynesthesia/Private/LKFSNRT.cpp
new file mode 100644
--- /dev/null
+++ b/Engine/Plugins/Runtime/AudioSynesthesia/Source/AudioSynesthesia/Private/LKFSNRT.cpp
@@ -0,0 +1,250 @@
+// Copyright Epic Games, Inc. All Rights Reserved.
+
+#include "LKFSNRT.h"
+
+#include "FindNearestByTimestamp.h"
+#include "LKFSNRTFactory.h"
+#include "InterpolateSorted.h"
+#include "AudioSynesthesiaLog.h"
+
+#include UE_INLINE_GENERATED_CPP_BY_NAME(LKFSNRT)
+
+namespace LKFSNRTPrivate
+{
+	float InterpolateLoudness(TArrayView<const FLKFSNRTResults> InLoudnessArray, const float InTimestamp)
+	{
+		int32 LowerIndex = INDEX_NONE;
+		int32 UpperIndex = INDEX_NONE;
+		float Alpha = 0.f;
+
+		Audio::GetInterpolationParametersAtTimestamp(InLoudnessArray, InTimestamp, LowerIndex, UpperIndex, Alpha);
+
+		if ((INDEX_NONE != LowerIndex) && (INDEX_NONE != UpperIndex))
+		{
+			return FMath::Lerp(InLoudnessArray[LowerIndex].Loudness, InLoudnessArray[UpperIndex].Loudness, Alpha);
+		}
+
+		return 0.f;
+	}
+}
+
+
+/***************************************************************************/
+/**********************    ULKFSNRTSettings     ************************/
+/***************************************************************************/
+
+ULKFSNRTSettings::ULKFSNRTSettings()
+{}
+
+TUniquePtr<Audio::IAnalyzerNRTSettings> ULKFSNRTSettings::GetSettings(const float InSampleRate, const int32 InNumChannels) const
+{
+	using namespace Audio;
+
+	TUniquePtr<FLKFSNRTSettings> OutSettings = MakeUnique<FLKFSNRTSettings>();
+
+	// lkfs is generally done in time domain, but the loudness analyzer works in 
+	// the spectral domain. Choosing this FFT small enough to allow for proper 
+	// spectral-temporal granularity. If the FFT size is too large it can become 
+	// larger than the AnalysisWindowDuration. If the FFT is too small, the perceptual
+	// weighting in the frequency domain looses accuracy. 
+	OutSettings->LoudnessAnalyzerSettings.FFTSize = 2048;
+	OutSettings->LoudnessAnalyzerSettings.WindowType = Audio::EWindowType::Blackman;
+	OutSettings->LoudnessAnalyzerSettings.MinAnalysisFrequency = 20.f;
+	OutSettings->LoudnessAnalyzerSettings.MaxAnalysisFrequency = 20000.f;
+	OutSettings->LoudnessAnalyzerSettings.WindowSize = -1;
+	// Old loudness analyzer settings had incorrect method for doing loudness metrics. To match LKFS standard we have to use a corrected scaling method.
+	OutSettings->LoudnessAnalyzerSettings.ScalingMethod = ELoudnessAnalyzerScalingMethod::Corrected;
+	OutSettings->LoudnessAnalyzerSettings.LoudnessCurveType = ELoudnessCurveType::K,
+
+	OutSettings->AnalysisPeriod = AnalysisPeriod;
+	OutSettings->AnalysisWindowDuration = AnalysisWindowDuration;
+	OutSettings->ShortTermLoudnessDuration = ShortTermLoudnessDuration;
+		
+	return OutSettings;
+}
+
+#if WITH_EDITOR
+FText ULKFSNRTSettings::GetAssetActionName() const
+{
+	return NSLOCTEXT("AssetTypeActions", "AssetTypeActions_AssetSoundSynesthesiaLKFSNRTSettings", "Synesthesia NRT Settings (LKFS / LUFS)");
+}
+
+UClass* ULKFSNRTSettings::GetSupportedClass() const
+{
+	return ULKFSNRTSettings::StaticClass();
+}
+#endif
+
+/***************************************************************************/
+/**********************        ULKFSNRT         ************************/
+/***************************************************************************/
+
+ULKFSNRT::ULKFSNRT()
+{
+	Settings = CreateDefaultSubobject<ULKFSNRTSettings>(TEXT("DefaultLKFSNRTSettings"));
+#if WITH_EDITOR
+	// Bind settings to audio analyze so changes to default settings will trigger analysis.
+	SetSettingsDelegate(Settings);
+#endif
+}
+
+void ULKFSNRT::GetLoudnessAtTime(const float InSeconds, float& OutLoudness) const
+{
+	GetChannelLoudnessAtTime(InSeconds, Audio::FLKFSNRTResult::ChannelIndexOverall, OutLoudness);
+}
+
+void ULKFSNRT::GetChannelLoudnessAtTime(const float InSeconds, const int32 InChannel, float& OutLoudness) const
+{
+	OutLoudness = 0.0f;
+
+	TSharedPtr<const Audio::FLKFSNRTResult, ESPMode::ThreadSafe> LoudnessResult = GetResult<Audio::FLKFSNRTResult>();
+
+	if (LoudnessResult.IsValid())
+	{
+		// The loudness result should never used here if it is not already sorted.
+		check(LoudnessResult->IsSortedChronologically());
+
+		if (!LoudnessResult->ContainsChannel(InChannel))
+		{
+			UE_LOG(LogAudioSynesthesia, Warning, TEXT("LKFSNRT does not contain channel %d"), InChannel);
+			return;
+		}
+
+		TArrayView<const FLKFSNRTResults> LoudnessArray = LoudnessResult->GetChannelLoudnessArray(InChannel);
+
+		OutLoudness = LKFSNRTPrivate::InterpolateLoudness(LoudnessArray, InSeconds);
+	}
+}
+
+TArray<FLKFSNRTResults> ULKFSNRT::GetLoudnessDataForChannel(int32 InChannel) const
+{
+	TSharedPtr<const Audio::FLKFSNRTResult, ESPMode::ThreadSafe> LoudnessResult = GetResult<Audio::FLKFSNRTResult>();
+	if (LoudnessResult.IsValid())
+	{
+		// The loudness result should never used here if it is not already sorted.
+		check(LoudnessResult->IsSortedChronologically());
+
+		if (!LoudnessResult->ContainsChannel(InChannel))
+		{
+			UE_LOG(LogAudioSynesthesia, Warning, TEXT("LKFSNRT does not contain channel %d"), InChannel);
+			return {};
+		}
+
+		return TArray<FLKFSNRTResults>(LoudnessResult->GetChannelLoudnessArray(InChannel));
+	}
+	return {};
+}
+
+TArray<FLKFSNRTResults> ULKFSNRT::GetLoudnessData() const
+{
+	return GetLoudnessDataForChannel(Audio::FLKFSNRTResult::ChannelIndexOverall);
+}
+
+FLKFSNRTResults ULKFSNRT::GetLoudnessDataForChannelAtTime(const float InSeconds, int32 InChannel) const
+{
+	TSharedPtr<const Audio::FLKFSNRTResult, ESPMode::ThreadSafe> LoudnessResult = GetResult<Audio::FLKFSNRTResult>();
+	if (LoudnessResult.IsValid())
+	{
+		// The loudness result should never used here if it is not already sorted.
+		check(LoudnessResult->IsSortedChronologically());
+
+		if (!LoudnessResult->ContainsChannel(InChannel))
+		{
+			UE_LOG(LogAudioSynesthesia, Warning, TEXT("LKFSNRT does not contain channel %d"), InChannel);
+			return {};
+		}
+
+		TArrayView<const FLKFSNRTResults> LoudnessArray = LoudnessResult->GetChannelLoudnessArray(InChannel);
+		if (const FLKFSNRTResults* Datum = Audio::FindNearestByTimestamp<const FLKFSNRTResults>(LoudnessArray, InSeconds))
+		{
+			return *Datum;
+		}
+	}
+	return {};
+}
+
+FLKFSNRTResults ULKFSNRT::GetLoudnessDataAtTime(const float InSeconds) const
+{
+	return GetLoudnessDataForChannelAtTime(InSeconds, Audio::FLKFSNRTResult::ChannelIndexOverall);
+}
+
+float ULKFSNRT::GetIntegratedLoudnessForChannel(int32 InChannel) const
+{
+	TSharedPtr<const Audio::FLKFSNRTResult, ESPMode::ThreadSafe> LoudnessResult = GetResult<Audio::FLKFSNRTResult>();
+	if (LoudnessResult.IsValid())
+	{
+		if (const FLKFSNRTAggregateStats* Stats = LoudnessResult->FindAggregateLoudnessStats(InChannel))
+		{
+			return Stats->IntegratedLoudness;
+		}
+		else
+		{
+			UE_LOG(LogAudioSynesthesia, Warning, TEXT("LKFSNRT does not contain integrated loudness for channel %d"), InChannel);
+		}
+	}
+	return FLKFSNRTAggregateStats::InvalidLoudness;
+}
+
+float ULKFSNRT::GetIntegratedLoudness() const
+{
+	return GetIntegratedLoudnessForChannel(Audio::FLKFSNRTResult::ChannelIndexOverall);
+}
+
+float ULKFSNRT::GetGatedLoudnessForChannel(int32 InChannel) const
+{
+	TSharedPtr<const Audio::FLKFSNRTResult, ESPMode::ThreadSafe> LoudnessResult = GetResult<Audio::FLKFSNRTResult>();
+	if (LoudnessResult.IsValid())
+	{
+		if (const FLKFSNRTAggregateStats* Stats = LoudnessResult->FindAggregateLoudnessStats(InChannel))
+		{
+			return Stats->GatedLoudness;
+		}
+		else
+		{
+			UE_LOG(LogAudioSynesthesia, Warning, TEXT("LKFSNRT does not contain gated loudness for channel %d"), InChannel);
+		}
+	}
+	return FLKFSNRTAggregateStats::InvalidLoudness;
+}
+
+float ULKFSNRT::GetGatedLoudness() const
+{
+	return GetGatedLoudnessForChannel(Audio::FLKFSNRTResult::ChannelIndexOverall);
+}
+
+TUniquePtr<Audio::IAnalyzerNRTSettings> ULKFSNRT::GetSettings(const float InSampleRate, const int32 InNumChannels) const
+{
+	TUniquePtr<Audio::IAnalyzerNRTSettings> AnalyzerSettings;
+
+	if (Settings)
+	{
+		AnalyzerSettings = Settings->GetSettings(InSampleRate, InNumChannels);	
+	}
+
+	return AnalyzerSettings;
+}
+
+#if WITH_EDITOR
+FText ULKFSNRT::GetAssetActionName() const
+{
+	return NSLOCTEXT("AssetTypeActions", "AssetTypeActions_AssetSoundSynesthesiaLKFSNRT", "Synesthesia NRT (LKFS / LUFS)");
+}
+
+UClass* ULKFSNRT::GetSupportedClass() const
+{
+	return ULKFSNRT::StaticClass();
+}
+
+bool ULKFSNRT::ShouldEventTriggerAnalysis(FPropertyChangedEvent& PropertyChangeEvent)
+{
+	return Settings != nullptr && Super::ShouldEventTriggerAnalysis(PropertyChangeEvent);
+}
+#endif // WITH_EDITOR
+
+FName ULKFSNRT::GetAnalyzerNRTFactoryName() const
+{
+	static const FName FactoryName(TEXT("LKFSNRTFactory"));
+	return FactoryName;
+}
+
+
diff --git a/Engine/Plugins/Runtime/AudioSynesthesia/Source/AudioSynesthesia/Private/LoudnessNRT.cpp b/Engine/Plugins/Runtime/AudioSynesthesia/Source/AudioSynesthesia/Private/LoudnessNRT.cpp
--- a/Engine/Plugins/Runtime/AudioSynesthesia/Source/AudioSynesthesia/Private/LoudnessNRT.cpp
+++ b/Engine/Plugins/Runtime/AudioSynesthesia/Source/AudioSynesthesia/Private/LoudnessNRT.cpp
@@ -7,7 +7,7 @@
 
 #include UE_INLINE_GENERATED_CPP_BY_NAME(LoudnessNRT)
 
-namespace
+namespace LoudnessNRTPrivate
 {
 	float InterpolateLoudness(TArrayView<const Audio::FLoudnessDatum> InLoudnessArray, const float InTimestamp)
 	{
@@ -15,7 +15,7 @@ namespace
 		int32 UpperIndex = INDEX_NONE;
 		float Alpha = 0.f;
 
-		GetInterpolationParametersAtTimestamp(InLoudnessArray, InTimestamp, LowerIndex, UpperIndex, Alpha);
+		Audio::GetInterpolationParametersAtTimestamp(InLoudnessArray, InTimestamp, LowerIndex, UpperIndex, Alpha);
 
 		if ((INDEX_NONE != LowerIndex) && (INDEX_NONE != UpperIndex))
 		{
@@ -46,7 +46,9 @@ TUniquePtr<Audio::IAnalyzerNRTSettings> ULoudnessNRTSettings::GetSettings(const
 	Settings->AnalysisPeriod = AnalysisPeriod;
 	Settings->MinAnalysisFrequency = MinimumFrequency;
 	Settings->MaxAnalysisFrequency = MaximumFrequency;
-
+	Settings->WindowType = Audio::EWindowType::Blackman;
+	Settings->ScalingMethod = Audio::ELoudnessAnalyzerScalingMethod::BackCompatible;
+		
 	switch (CurveType)
 	{
 		case ELoudnessNRTCurveTypeEnum::A:
@@ -127,7 +129,7 @@ void ULoudnessNRT::GetChannelLoudnessAtTime(const float InSeconds, const int32 I
 
 		const TArray<Audio::FLoudnessDatum>& LoudnessArray = LoudnessResult->GetChannelLoudnessArray(InChannel);
 
-		OutLoudness = InterpolateLoudness(LoudnessArray, InSeconds);
+		OutLoudness = LoudnessNRTPrivate::InterpolateLoudness(LoudnessArray, InSeconds);
 	}
 }
 
@@ -155,7 +157,7 @@ void ULoudnessNRT::GetNormalizedChannelLoudnessAtTime(const float InSeconds, con
 
 		const TArray<Audio::FLoudnessDatum>& LoudnessArray = LoudnessResult->GetChannelLoudnessArray(InChannel);
 
-		OutLoudness = InterpolateLoudness(LoudnessArray, InSeconds);
+		OutLoudness = LoudnessNRTPrivate::InterpolateLoudness(LoudnessArray, InSeconds);
 
 		// Subtract offset
 		OutLoudness -= Settings->NoiseFloorDb;
diff --git a/Engine/Plugins/Runtime/AudioSynesthesia/Source/AudioSynesthesiaCore/AudioSynesthesiaCore.Build.cs b/Engine/Plugins/Runtime/AudioSynesthesia/Source/AudioSynesthesiaCore/AudioSynesthesiaCore.Build.cs
--- a/Engine/Plugins/Runtime/AudioSynesthesia/Source/AudioSynesthesiaCore/AudioSynesthesiaCore.Build.cs
+++ b/Engine/Plugins/Runtime/AudioSynesthesia/Source/AudioSynesthesiaCore/AudioSynesthesiaCore.Build.cs
@@ -18,11 +18,12 @@ namespace UnrealBuildTool.Rules
             PublicDependencyModuleNames.AddRange(
 				new string[] {
                     "Core",
+					"CoreUObject",
 					"SignalProcessing",
                 }
             );
 
-			PrivateIncludePathModuleNames.AddRange(
+			PrivateDependencyModuleNames.AddRange(
 				new string[] {
 					"AudioMixerCore"
 				}
diff --git a/Engine/Plugins/Runtime/AudioSynesthesia/Source/AudioSynesthesiaCore/Private/AudioSynesthesiaCoreModule.cpp b/Engine/Plugins/Runtime/AudioSynesthesia/Source/AudioSynesthesiaCore/Private/AudioSynesthesiaCoreModule.cpp
--- a/Engine/Plugins/Runtime/AudioSynesthesia/Source/AudioSynesthesiaCore/Private/AudioSynesthesiaCoreModule.cpp
+++ b/Engine/Plugins/Runtime/AudioSynesthesia/Source/AudioSynesthesiaCore/Private/AudioSynesthesiaCoreModule.cpp
@@ -8,6 +8,7 @@
 #include "HAL/LowLevelMemTracker.h"
 #include "LoudnessFactory.h"
 #include "LoudnessNRTFactory.h"
+#include "LKFSNRTFactory.h"
 #include "MeterFactory.h"
 #include "Modules/ModuleManager.h"
 #include "OnsetNRTFactory.h"
@@ -27,6 +28,7 @@ namespace Audio
 
 				// Register factories on startup
 				IModularFeatures::Get().RegisterModularFeature(FLoudnessNRTFactory::GetModularFeatureName(), &LoudnessNRTFactory);
+				IModularFeatures::Get().RegisterModularFeature(FLKFSNRTFactory::GetModularFeatureName(), &LKFSNRTFactory);
 				IModularFeatures::Get().RegisterModularFeature(FConstantQNRTFactory::GetModularFeatureName(), &ConstantQNRTFactory);
 				IModularFeatures::Get().RegisterModularFeature(FOnsetNRTFactory::GetModularFeatureName(), &OnsetNRTFactory);
 
@@ -41,6 +43,7 @@ namespace Audio
 				LLM_SCOPE_BYNAME(TEXT("Audio/AudioAnalysis"));
 				// Unregister factories on shutdown
 				IModularFeatures::Get().UnregisterModularFeature(FLoudnessNRTFactory::GetModularFeatureName(), &LoudnessNRTFactory);
+				IModularFeatures::Get().UnregisterModularFeature(FLKFSNRTFactory::GetModularFeatureName(), &LKFSNRTFactory);
 				IModularFeatures::Get().UnregisterModularFeature(FConstantQNRTFactory::GetModularFeatureName(), &ConstantQNRTFactory);
 				IModularFeatures::Get().UnregisterModularFeature(FOnsetNRTFactory::GetModularFeatureName(), &OnsetNRTFactory);
 
@@ -52,6 +55,7 @@ namespace Audio
 
 		private:
 			FLoudnessNRTFactory LoudnessNRTFactory;
+			FLKFSNRTFactory LKFSNRTFactory;
 			FConstantQNRTFactory ConstantQNRTFactory;
 			FOnsetNRTFactory OnsetNRTFactory;
 
diff --git a/Engine/Plugins/Runtime/AudioSynesthesia/Source/AudioSynesthesiaCore/Private/LKFSNRTFactory.cpp b/Engine/Plugins/Runtime/AudioSynesthesia/Source/AudioSynesthesiaCore/Private/LKFSNRTFactory.cpp
new file mode 100644
--- /dev/null
+++ b/Engine/Plugins/Runtime/AudioSynesthesia/Source/AudioSynesthesiaCore/Private/LKFSNRTFactory.cpp
@@ -0,0 +1,495 @@
+// Copyright Epic Games, Inc. All Rights Reserved.
+
+#include "LKFSNRTFactory.h"
+
+#include "Algo/IsSorted.h"
+#include "AudioSynesthesiaCoreLog.h"
+#include "DSP/SlidingWindow.h"
+
+namespace Audio
+{
+	namespace LKFSNRTFactoryPrivate
+	{
+		float CalculateGatedLoudnessPass(TArrayView<const FLKFSNRTResults> InLoudnessData, float InLoudnessThreshold)
+		{
+			// Get average loudness of values about gating threshold
+			double AccumEnergy = 0.f;
+			int32 AccumNum = 0;
+
+			for (const FLKFSNRTResults& Datum : InLoudnessData)
+			{
+				if (Datum.Loudness > InLoudnessThreshold)
+				{
+					AccumEnergy += Datum.Energy;
+					AccumNum++;
+				}
+			}
+
+			if (AccumNum > 0)
+			{
+				float AverageEnergy = (float)(AccumEnergy / (double)AccumNum);
+				return FLoudnessAnalyzer::ConvertPerceptualEnergyToLKFS(AverageEnergy);
+			}
+			else
+			{
+				// No values exceeded threshold. Using the threshold seems like a
+				// reasonable result.
+				return InLoudnessThreshold;
+			}
+		}
+
+		float CalculateGatedLoudness(TArrayView<const FLKFSNRTResults> InLoudnessData)
+		{
+			if (InLoudnessData.Num() < 1)
+			{
+				return FLKFSNRTAggregateStats::InvalidLoudness;;
+			}
+			// Gated Loudness only examines loudnesses above a certain gating loudness. 
+			//
+			// It is a 2 pass algorithm 
+			// Pass 1: Use a set loudness threshold of -70lkfs to determine an initial gated loudness
+			// Pass 2: Use the initial gated loudness to derive a new loudness threshold and calculate the gate loudness again. 
+			constexpr float InitLoudnessThreshold = -70.f;
+
+			const float GatedLoudnessPass1 = CalculateGatedLoudnessPass(InLoudnessData, InitLoudnessThreshold);
+			const float NewThreshold = GatedLoudnessPass1 - 10.f;
+			return CalculateGatedLoudnessPass(InLoudnessData, NewThreshold);
+		}
+		
+		void CalculateShortTermAndIntegratedLoudness(float InShortTermLoudnessDuration, TArrayView<FLKFSNRTResults> InOutSortedLoudnessData, FLKFSNRTAggregateStats& InOutStats)
+		{
+			check(InShortTermLoudnessDuration > 0.f);
+			checkSlow(Algo::IsSortedBy(InOutSortedLoudnessData, [](const FLKFSNRTResults& InDatum) { return InDatum.Timestamp; }));
+
+			if (InOutSortedLoudnessData.Num() < 1)
+			{
+				InOutStats.IntegratedLoudness = FLKFSNRTAggregateStats::InvalidLoudness;
+				return;
+			}
+
+			// Track a short term window. To find the energy inside of the short term window,
+			// accumulate the energy over all the data. For a given short term window, the 
+			// energy in that window can be determined by subtracting the tail accumulated
+			// energy from the head accumulated energy. 
+			
+
+			double HeadAccumEnergy = 0; // Total accumulated energy at head
+			double ShortTermTailAccumEnergy = 0; // Total accumulated energy at tail
+			const float FirstTimestamp = InOutSortedLoudnessData[0].Timestamp;
+			double ShortTermLoudnessNum = 0;
+			
+			const FLKFSNRTResults* ShortTermLoudnessDatumBegin = &InOutSortedLoudnessData[0];
+
+			for (int32 Index = 0; Index < InOutSortedLoudnessData.Num(); Index++)
+			{
+				FLKFSNRTResults& LoudnessDatum = InOutSortedLoudnessData[Index];
+				// Acummulated loudness at head
+				HeadAccumEnergy += LoudnessDatum.Energy;
+				
+				// Only accumulated loudness at tail if we are within the desired window
+				float MinTimestamp = LoudnessDatum.Timestamp - InShortTermLoudnessDuration;
+				if (MinTimestamp < FirstTimestamp)
+				{
+					// Grow the number of elements in the window since the difference in timestamps
+					// between the head and tail of the data are less than the window size
+					ShortTermLoudnessNum = (float)(Index + 1);
+				}
+				else
+				{
+					ShortTermTailAccumEnergy += ShortTermLoudnessDatumBegin->Energy;
+					ShortTermLoudnessDatumBegin++;
+				}
+
+				checkf((LoudnessDatum.Timestamp - ShortTermLoudnessDatumBegin->Timestamp) <= (InShortTermLoudnessDuration + UE_KINDA_SMALL_NUMBER), TEXT("Invalid short term loudness window. The algorithm for integrated and short term loudness assumes a data to be regularly spaced in time"));
+
+				// Short term loudness is the accumulated loudness divided the number of elements in the window.
+				float ShortTermEnergy = (float)((HeadAccumEnergy - ShortTermTailAccumEnergy) / ShortTermLoudnessNum);
+				LoudnessDatum.ShortTermLoudness = FLoudnessAnalyzer::ConvertPerceptualEnergyToLKFS(ShortTermEnergy);
+			}
+
+			// Integrated loudness is the average over the entire file.
+			float IntegratedEnergy = (float)(HeadAccumEnergy / (double)InOutSortedLoudnessData.Num());
+			InOutStats.IntegratedLoudness = FLoudnessAnalyzer::ConvertPerceptualEnergyToLKFS(IntegratedEnergy);
+		}
+
+
+		// Calculates ShortTermLoudness (part of FLKFSNRTResults) and FLKFSNRTAggregateStats
+		void CalculateAggregateLoudnessStats(float InShortTermLoudnessDuration, TArrayView<FLKFSNRTResults> InOutSortedLoudnessData, FLKFSNRTAggregateStats& InOutStats)
+		{
+			CalculateShortTermAndIntegratedLoudness(InShortTermLoudnessDuration, InOutSortedLoudnessData, InOutStats);
+			InOutStats.GatedLoudness = CalculateGatedLoudness(InOutSortedLoudnessData);
+		}
+	}
+}
+
+namespace Audio
+{
+	/************************************************************************/
+	/********************** FLKFSNRTResult ******************************/
+	/************************************************************************/
+
+	// Use for overall loudness storage
+	const int32 FLKFSNRTResult::ChannelIndexOverall = -1;
+
+	FLKFSNRTResult::FLKFSNRTResult()
+	{
+	}
+
+	void FLKFSNRTResult::AddChannel(int32 InChannelIndex)
+	{
+		if (!Data.ChannelData.Contains(InChannelIndex))
+		{
+			Data.ChannelData.Add(InChannelIndex);
+		}
+	}
+
+	void FLKFSNRTResult::Add(const FLKFSNRTResults& InDatum)
+	{
+		FLKFSNRTChannelData& ChannelDatum = Data.ChannelData.FindOrAdd(InDatum.Channel);
+		// Store loudness data in appropriate channel
+		ChannelDatum.LoudnessArray.Add(InDatum);
+
+		// Mark as not sorted
+		Data.bIsSortedChronologically = false;
+	}
+
+	bool FLKFSNRTResult::ContainsChannel(int32 InChannelIndex) const
+	{
+		return Data.ChannelData.Contains(InChannelIndex);
+	}
+
+	TArrayView<const FLKFSNRTResults> FLKFSNRTResult::GetChannelLoudnessArray(int32 ChannelIdx) const
+	{
+		if (const FLKFSNRTChannelData* ChannelDatum = Data.ChannelData.Find(ChannelIdx))
+		{
+			return ChannelDatum->LoudnessArray;
+		}
+		return {};
+	}
+
+	void FLKFSNRTResult::Serialize(FArchive& Archive)
+	{
+		FLKFSNRTWaveData::StaticStruct()->SerializeItem(Archive, &Data, nullptr);
+	}
+
+	TArrayView<FLKFSNRTResults> FLKFSNRTResult::GetMutableChannelLoudnessArray(int32 ChannelIdx)
+	{
+		if (FLKFSNRTChannelData* ChannelDatum = Data.ChannelData.Find(ChannelIdx))
+		{
+			return ChannelDatum->LoudnessArray;
+		}
+		return {};
+	}
+
+	TArrayView<const FLKFSNRTResults> FLKFSNRTResult::GetLoudnessArray() const
+	{
+		return GetChannelLoudnessArray(ChannelIndexOverall);
+	}
+
+	const FLKFSNRTAggregateStats* FLKFSNRTResult::FindAggregateLoudnessStats(int32 InChannelIdx) const
+	{
+		if (const FLKFSNRTChannelData* ChannelDatum = Data.ChannelData.Find(InChannelIdx))
+		{
+			return &ChannelDatum->AggregateStats;
+		}
+		return nullptr;
+	}
+
+	void FLKFSNRTResult::SetAggregateLoudnessStats(int32 InChannelIdx, FLKFSNRTAggregateStats InStats)
+	{
+		Data.ChannelData.FindOrAdd(InChannelIdx).AggregateStats = InStats;
+	}
+
+	void FLKFSNRTResult::GetChannels(TArray<int32>& OutChannels) const
+	{
+		Data.ChannelData.GetKeys(OutChannels);
+	}
+
+	float FLKFSNRTResult::GetDurationInSeconds() const 
+	{
+		return Data.DurationInSeconds;
+	}
+
+	void FLKFSNRTResult::SetDurationInSeconds(float InDuration)
+	{
+		Data.DurationInSeconds = InDuration;
+	}
+
+	bool FLKFSNRTResult::IsSortedChronologically() const
+	{
+		return Data.bIsSortedChronologically;
+	}
+
+	void FLKFSNRTResult::SortChronologically()
+	{
+		for (auto& KeyValue : Data.ChannelData)
+		{
+			KeyValue.Value.LoudnessArray.Sort([](const FLKFSNRTResults& A, const FLKFSNRTResults& B) { return A.Timestamp < B.Timestamp; });
+		}
+
+		Data.bIsSortedChronologically = true;
+	}
+
+	/************************************************************************/
+	/********************** FLKFSNRTWorker ******************************/
+	/************************************************************************/
+
+	FLKFSNRTWorker::FLKFSNRTWorker(const FAnalyzerNRTParameters& InParams, const FLKFSNRTSettings& InAnalyzerSettings)
+	: NumChannels(InParams.NumChannels)
+	, SampleRate(InParams.SampleRate)
+	, ShortTermLoudnessDuration(InAnalyzerSettings.ShortTermLoudnessDuration)
+	{
+		check(NumChannels > 0);
+		check(SampleRate > 0.f);
+		check(ShortTermLoudnessDuration > 0.f);
+
+		Analyzer = MakeUnique<FMultichannelLoudnessAnalyzer>(InParams.SampleRate, InAnalyzerSettings.LoudnessAnalyzerSettings);
+
+		int32 NumSlidingWindowFrames = Analyzer->GetSettings().WindowSize;
+
+		int32 NumAnalysisFrames = FMath::CeilToInt(InAnalyzerSettings.AnalysisWindowDuration * InParams.SampleRate);
+		int32 NumAnalysisHopFrames = FMath::CeilToInt(InAnalyzerSettings.AnalysisPeriod * InParams.SampleRate);
+
+		// There may be multiple calls to the LoudnessAnalyzer to produce a single 
+		// FLoundessDatum. The sliding window hop is tuned here so that it best 
+		// matches the desired AnalysisPeriod while maintaining between a 25% to 
+		// 75% window overlap.
+		TuneSlidingWindwoHopSize(NumAnalysisHopFrames, NumSlidingWindowFrames, NumSlidingWindowHopFrames, NumAnalysisHopWindows);
+
+		// Determine how many windows to analyze per loudness result
+		NumAnalysisWindows = FMath::Max(1, ((NumAnalysisFrames - (NumSlidingWindowFrames / 2)) / NumSlidingWindowHopFrames) + 1);
+
+		// Include NumChannels in calculating window size because windows are generated 
+		// with interleaved samples and deinterleaved later. 
+		int32 NumSlidingWindowSamples = NumSlidingWindowFrames * NumChannels;
+		int32 NumSlidingWindowHopSamples = NumSlidingWindowHopFrames * NumChannels;
+
+		InternalBuffer = MakeUnique<TSlidingBuffer<float>>(NumSlidingWindowSamples, NumSlidingWindowHopSamples);
+		InternalWindow.AddUninitialized(NumSlidingWindowSamples);
+		ChannelPerceptualEnergy.AddZeroed(NumChannels * NumAnalysisWindows);
+		OverallPerceptualEnergy.AddZeroed(NumAnalysisWindows);
+	}
+
+	void FLKFSNRTWorker::Analyze(TArrayView<const float> InAudio, IAnalyzerNRTResult* OutResult) 
+	{
+		NumFrames += InAudio.Num() / NumChannels;
+
+		// Assume that outer layers ensured that this is of correct type.
+		FLKFSNRTResult* LoudnessResult = static_cast<FLKFSNRTResult*>(OutResult);
+
+		check(nullptr != LoudnessResult);
+		
+		TAutoSlidingWindow<float> SlidingWindow(*InternalBuffer, InAudio, InternalWindow);
+
+		AnalyzeSlidingWindow(SlidingWindow, *LoudnessResult);
+	}
+
+	// Called when analysis of audio asset is complete. 
+	void FLKFSNRTWorker::Finalize(IAnalyzerNRTResult* OutResult) 
+	{
+		FLKFSNRTResult* LoudnessResult = static_cast<FLKFSNRTResult*>(OutResult);
+
+		check(nullptr != LoudnessResult);
+
+		TArray<float> EmptyArray;
+		constexpr bool bFlushAudioWithZeroes = false;
+		TAutoSlidingWindow<float> SlidingWindow(*InternalBuffer, EmptyArray, InternalWindow, bFlushAudioWithZeroes);
+
+		AnalyzeSlidingWindow(SlidingWindow, *LoudnessResult);
+
+		// Finalize all the channels now that we have finished analyzing audio.
+		FinalizeChannels(*LoudnessResult);
+
+		// Prepare result for fast lookup.
+		LoudnessResult->SortChronologically(); 
+
+		LoudnessResult->SetDurationInSeconds(static_cast<float>(NumFrames) / SampleRate);
+
+		// Update integrated, gated and short term loudness. 
+		CalculateAggregateLoudnessStats(*LoudnessResult);
+
+		// Reset internal counters
+		NumAnalyzedBuffers = 0; 
+		NumFrames = 0;
+
+		// Reset sliding buffer
+		InternalBuffer->Reset();
+	}
+
+	void FLKFSNRTWorker::FinalizeChannels(FLKFSNRTResult& OutResult)
+	{
+		for (int32 ChannelIdx = 0; ChannelIdx < NumChannels; ChannelIdx++)
+		{
+			if (!OutResult.ContainsChannel(ChannelIdx))
+			{
+				OutResult.AddChannel(ChannelIdx);
+			}
+		}
+
+		if (!OutResult.ContainsChannel(FLKFSNRTResult::ChannelIndexOverall))
+		{
+			OutResult.AddChannel(FLKFSNRTResult::ChannelIndexOverall);
+		}
+	}
+
+	void FLKFSNRTWorker::TuneSlidingWindwoHopSize(const int32 InAnalysisHopFrames, const int32 InWindowFrames, int32& OutBestWindowHopFrames, int32& OutBestAnalysisHopWindows) const
+	{
+		// The minimum sliding window hop size is the lesser of 
+		// - 25% of the sliding window size
+		// - The analysis hop size
+		int32 MinNumSlidingWindowHopFrames = FMath::Max(1, InWindowFrames / 4);
+		if (MinNumSlidingWindowHopFrames > InAnalysisHopFrames)
+		{
+			MinNumSlidingWindowHopFrames = InAnalysisHopFrames;
+		}
+
+		// The maximum sliding window hop size is the lesser of
+		// - 75% of the sliding window size
+		// - The analysis hop size
+		int32 MaxNumSlidingWindowHopFrames = FMath::Max(1, (3 * InWindowFrames) / 4);
+		if (MaxNumSlidingWindowHopFrames > InAnalysisHopFrames)
+		{
+			MaxNumSlidingWindowHopFrames = InAnalysisHopFrames;
+		}
+
+		// Iterate through the possible sizes to find the one that fits best with our settings.
+		int32 BestWindowHopFrames = -1;
+		int32 BestAnalysisHopWindows = -1;
+		int32 BestMismatch = InAnalysisHopFrames + 2; 
+
+		for (int32 CandidateWindowHopFrames = MinNumSlidingWindowHopFrames; CandidateWindowHopFrames <= MaxNumSlidingWindowHopFrames; CandidateWindowHopFrames++)
+		{
+			int32 CandidateAnalysisHopWindows = InAnalysisHopFrames / CandidateWindowHopFrames;
+			// The match quality is determined by how neatly the sliding window hop size fits into the anlaysis window hop size. 
+			int32 CandidateMismatch = FMath::Abs<int32>(InAnalysisHopFrames - (CandidateWindowHopFrames * CandidateAnalysisHopWindows));
+			if (CandidateMismatch < BestMismatch)
+			{
+				BestMismatch = CandidateMismatch;
+				BestWindowHopFrames = CandidateWindowHopFrames;
+				BestAnalysisHopWindows = CandidateAnalysisHopWindows;
+			}
+		}
+
+		// The algorithm above should always find a valid window hop
+		check(BestWindowHopFrames != -1);
+
+		OutBestWindowHopFrames = BestWindowHopFrames;
+		OutBestAnalysisHopWindows = BestAnalysisHopWindows;
+	}
+
+	void FLKFSNRTWorker::AnalyzeSlidingWindow(TAutoSlidingWindow<float>& InSlidingWindow, FLKFSNRTResult& OutResult)
+	{
+		check(NumAnalysisWindows > 0);
+		check(NumAnalysisHopWindows > 0);
+		
+		for (const TArray<float>& Window : InSlidingWindow)
+		{
+			AnalyzeWindow(Window, OutResult);
+
+			if ((NumAnalyzedBuffers >= NumAnalysisWindows) && ((NumAnalyzedBuffers - NumAnalysisWindows) % NumAnalysisHopWindows) == 0)
+			{
+				AddCurrentLoudnessDataToResult(OutResult);
+			}
+		}
+	}
+
+	void FLKFSNRTWorker::AnalyzeWindow(TArrayView<const float> InWindow, FLKFSNRTResult& OutResult)
+	{
+		const int32 WrappedWindowIdx = (NumAnalyzedBuffers % NumAnalysisWindows);
+		const int32 ChannelPerceptualEnergyWritePos = WrappedWindowIdx * NumChannels;
+
+		// Calculate perceptual energy
+		OverallPerceptualEnergy[WrappedWindowIdx] = Analyzer->CalculatePerceptualEnergy(InWindow, NumChannels, MakeArrayView<float>(ChannelPerceptualEnergy.GetData() + ChannelPerceptualEnergyWritePos, NumChannels));
+
+		// update counters
+		NumAnalyzedBuffers++;
+	}
+
+	void FLKFSNRTWorker::AddCurrentLoudnessDataToResult(FLKFSNRTResult& OutResult)
+	{
+		// Accounting for all the frames positions analyzed in creating the 
+		// data, determine the center frame. 
+		const int32 CenterFramePos = ((2 * NumAnalyzedBuffers - NumAnalysisWindows) * NumSlidingWindowHopFrames + Analyzer->GetSettings().WindowSize) / 2;
+		const float Timestamp = CenterFramePos / SampleRate;
+
+		// Compute loudness per channel
+		for (int32 ChannelIdx = 0; ChannelIdx < NumChannels; ChannelIdx++)
+		{
+			// Determine average perceptual energy over analysis window
+			float PerceptualEnergy = 0.f;
+			for (int32 Pos = ChannelIdx; Pos < ChannelPerceptualEnergy.Num(); Pos += NumChannels)
+			{
+				PerceptualEnergy += ChannelPerceptualEnergy[Pos];
+			}
+			PerceptualEnergy /= (float)NumAnalysisWindows;
+
+			// Convert to loudness.
+			const float PerceptualLoudness = FLoudnessAnalyzer::ConvertPerceptualEnergyToLKFS(PerceptualEnergy);
+			OutResult.Add(FLKFSNRTResults{.Channel=ChannelIdx, .Timestamp=Timestamp, .Energy=PerceptualEnergy, .Loudness=PerceptualLoudness});
+		}
+
+		{
+			// Compute loudness overall
+			float PerceptualEnergy = 0.f;
+			for (int32 Pos = 0; Pos < OverallPerceptualEnergy.Num(); Pos++)
+			{
+				PerceptualEnergy += OverallPerceptualEnergy[Pos];
+			}
+			PerceptualEnergy /= (float)NumAnalysisWindows;
+
+			const float PerceptualLoudness = FLoudnessAnalyzer::ConvertPerceptualEnergyToLKFS(PerceptualEnergy);
+			OutResult.Add(FLKFSNRTResults{.Channel=FLKFSNRTResult::ChannelIndexOverall,  .Timestamp=Timestamp, .Energy=PerceptualEnergy, .Loudness=PerceptualLoudness});
+		}
+	}
+
+	void FLKFSNRTWorker::CalculateAggregateLoudnessStats(FLKFSNRTResult& InOutResult) const
+	{
+		checkf(InOutResult.IsSortedChronologically(), TEXT("Results should be presorted in chronological order"));
+
+		TArray<int32> Channels;
+		InOutResult.GetChannels(Channels);
+
+		for (int32 ChannelIndex : Channels)
+		{
+			FLKFSNRTAggregateStats AggregateStats;
+
+			LKFSNRTFactoryPrivate::CalculateAggregateLoudnessStats(ShortTermLoudnessDuration, InOutResult.GetMutableChannelLoudnessArray(ChannelIndex), AggregateStats);
+
+			InOutResult.SetAggregateLoudnessStats(ChannelIndex, AggregateStats);
+		}
+	}
+
+
+	/************************************************************************/
+	/************************* FLKFSNRTFactory **************************/
+	/************************************************************************/
+
+	/** Name of specific analyzer type. */
+	FName FLKFSNRTFactory::GetName() const 
+	{
+		static FName FactoryName(TEXT("LKFSNRTFactory"));
+		return FactoryName;
+	}
+
+	/** Human readable name of analyzer. */
+	FString FLKFSNRTFactory::GetTitle() const
+	{
+		return TEXT("LKFS Analyzer Non-Real-Time");
+	}
+
+	TUniquePtr<IAnalyzerNRTResult> FLKFSNRTFactory::NewResult() const
+	{
+		TUniquePtr<FLKFSNRTResult> Result = MakeUnique<FLKFSNRTResult>();
+		return Result;
+	}
+
+	TUniquePtr<IAnalyzerNRTWorker> FLKFSNRTFactory::NewWorker(const FAnalyzerNRTParameters& InParams, const IAnalyzerNRTSettings* InSettings) const
+	{
+		const FLKFSNRTSettings* LoudnessSettings = static_cast<const FLKFSNRTSettings*>(InSettings);
+
+		check(nullptr != LoudnessSettings);
+
+		return MakeUnique<FLKFSNRTWorker>(InParams, *LoudnessSettings);
+	}
+}
+
diff --git a/Engine/Plugins/Runtime/AudioSynesthesia/Source/AudioSynesthesiaCore/Private/LoudnessAnalyzer.cpp b/Engine/Plugins/Runtime/AudioSynesthesia/Source/AudioSynesthesiaCore/Private/LoudnessAnalyzer.cpp
--- a/Engine/Plugins/Runtime/AudioSynesthesia/Source/AudioSynesthesiaCore/Private/LoudnessAnalyzer.cpp
+++ b/Engine/Plugins/Runtime/AudioSynesthesia/Source/AudioSynesthesiaCore/Private/LoudnessAnalyzer.cpp
@@ -5,10 +5,44 @@
 #include "AudioMixer.h"
 #include "DSP/DeinterleaveView.h"
 #include "DSP/FloatArrayMath.h"
+#include "GenericPlatform/GenericPlatformMath.h"
 
 namespace LoudnessAnalyzerPrivate
 {
 	static const float Log10Scale = 1.0f / FMath::Loge(10.f);
+
+	float GetChannelWeight(EAudioMixerChannel::Type InType) 
+	{
+		static_assert(static_cast<int32>(EAudioMixerChannel::ChannelTypeCount) == 19, "Possibly missing channel type");
+		switch (InType)
+		{
+			case EAudioMixerChannel::LowFrequency:
+				return 0.f; // LFE channels ignored
+				
+			case EAudioMixerChannel::BackLeft:
+			case EAudioMixerChannel::BackRight:
+			case EAudioMixerChannel::SideLeft:
+			case EAudioMixerChannel::SideRight:
+			// If elevation angle is less than 30 degrees and azimuth is between 60 and 120, set channel weight to sqrt(2)
+				return FMath::Sqrt(2.f);
+
+			case EAudioMixerChannel::FrontLeft:
+			case EAudioMixerChannel::FrontRight:
+			case EAudioMixerChannel::FrontCenter:
+			case EAudioMixerChannel::FrontLeftOfCenter:
+			case EAudioMixerChannel::FrontRightOfCenter:
+			case EAudioMixerChannel::BackCenter:
+			case EAudioMixerChannel::TopCenter:
+			case EAudioMixerChannel::TopFrontLeft:
+			case EAudioMixerChannel::TopFrontCenter:
+			case EAudioMixerChannel::TopFrontRight:
+			case EAudioMixerChannel::TopBackLeft:
+			case EAudioMixerChannel::TopBackCenter:
+			case EAudioMixerChannel::TopBackRight:
+			default:
+				return 1.f;
+		}
+	}
 }
 
 namespace Audio
@@ -341,13 +375,6 @@ namespace Audio
     FMultichannelLoudnessAnalyzer::FMultichannelLoudnessAnalyzer(float InSampleRate, const FLoudnessAnalyzerSettings& InSettings)
     :   Analyzer(InSampleRate, InSettings)
     {
-        // Initialize all channel weights to default to 1.0f
-        ChannelWeights.Init(1.0f, EAudioMixerChannel::ChannelTypeCount);
-
-        // If elevation angle is less than 30 degrees and azimuth is between 60 and 120, set channel weight to 1.5f. 
-        ChannelWeights[EAudioMixerChannel::BackLeft] = 1.5f;
-        ChannelWeights[EAudioMixerChannel::BackRight] = 1.5f;
-
         MonoBuffer.AddUninitialized(Analyzer.GetSettings().WindowSize);
     }
 
@@ -357,6 +384,17 @@ namespace Audio
 
         OutChannelEnergies.Reset(InNumChannels);
         OutChannelEnergies.AddUninitialized(InNumChannels);
+        return CalculatePerceptualEnergy(InView, InNumChannels, TArrayView<float>(OutChannelEnergies));
+    }
+
+    float FMultichannelLoudnessAnalyzer::CalculatePerceptualEnergy(TArrayView<const float> InView, const int32 InNumChannels, TArrayView<float> OutChannelEnergies)
+    {
+		check(InView.Num() == (InNumChannels * Analyzer.GetSettings().WindowSize));
+		check(OutChannelEnergies.Num() == InNumChannels);
+		if (ChannelWeights.Num() != InNumChannels)
+		{
+			InitializeChannelWeights(InNumChannels);
+		}
 
         // Deinterleave audio and send to mono loudness analyzer.
 		TAutoDeinterleaveView<float, FAudioBufferAlignedAllocator> DeinterleaveView(InView, MonoBuffer, InNumChannels);
@@ -376,6 +414,13 @@ namespace Audio
     }
 
     float FMultichannelLoudnessAnalyzer::CalculateLoudness(TArrayView<const float> InView, const int32 InNumChannels, TArray<float>& OutChannelLoudness)
+    {
+		OutChannelLoudness.Reset();
+		OutChannelLoudness.AddUninitialized(InNumChannels);
+		return CalculateLoudness(InView, InNumChannels, TArrayView<float>(OutChannelLoudness));
+    }
+
+    float FMultichannelLoudnessAnalyzer::CalculateLoudness(TArrayView<const float> InView, const int32 InNumChannels, TArrayView<float> OutChannelLoudness)
     {
         // Convert perceptual energy to loudness 
         float Energy = CalculatePerceptualEnergy(InView, InNumChannels, OutChannelLoudness);
@@ -386,6 +431,24 @@ namespace Audio
         return FLoudnessAnalyzer::ConvertPerceptualEnergyToLoudness(Energy);
     }
 
+	void FMultichannelLoudnessAnalyzer::InitializeChannelWeights(int32 InNumChannels)
+	{
+		ChannelWeights.SetNum(InNumChannels);
+		for (int32 ChannelIdx = 0; ChannelIdx < InNumChannels; ChannelIdx++)
+		{
+			EAudioMixerChannel::Type ChannelType = EAudioMixerChannel::DefaultChannel;
+
+			if (ensure(IAudioMixerPlatformInterface::GetChannelTypeAtIndex(ChannelIdx, ChannelType)))
+			{
+				ChannelWeights[ChannelIdx] = LoudnessAnalyzerPrivate::GetChannelWeight(ChannelType);
+			}
+			else
+			{
+				ChannelWeights[ChannelIdx] = 1.f;
+			}
+		}
+	}
+
     const FLoudnessAnalyzerSettings& FMultichannelLoudnessAnalyzer::GetSettings() const
     {
         return Analyzer.GetSettings();
diff --git a/Engine/Plugins/Runtime/AudioSynesthesia/Source/AudioSynesthesiaCore/Private/LoudnessFactory.cpp b/Engine/Plugins/Runtime/AudioSynesthesia/Source/AudioSynesthesiaCore/Private/LoudnessFactory.cpp
--- a/Engine/Plugins/Runtime/AudioSynesthesia/Source/AudioSynesthesiaCore/Private/LoudnessFactory.cpp
+++ b/Engine/Plugins/Runtime/AudioSynesthesia/Source/AudioSynesthesiaCore/Private/LoudnessFactory.cpp
@@ -75,10 +75,11 @@ namespace Audio
 	void FLoudnessWorker::AnalyzeWindow(TArrayView<const float> InWindow, FLoudnessResult& OutResult)
 	{
 		TArray<float> ChannelPerceptualEnergy;
+		ChannelPerceptualEnergy.AddZeroed(NumChannels);
 		float OverallPerceptualEnergy = 0.0f;
 
 		// Calculate perceptual energy
-		OverallPerceptualEnergy = Analyzer->CalculatePerceptualEnergy(InWindow, NumChannels, ChannelPerceptualEnergy);
+		OverallPerceptualEnergy = Analyzer->CalculatePerceptualEnergy(InWindow, NumChannels, TArrayView<float>(ChannelPerceptualEnergy));
 
 		// Convert to perceptual energy to loudness
 		TArray<float> ChannelLoudness;
diff --git a/Engine/Plugins/Runtime/AudioSynesthesia/Source/AudioSynesthesiaCore/Private/LoudnessNRTFactory.cpp b/Engine/Plugins/Runtime/AudioSynesthesia/Source/AudioSynesthesiaCore/Private/LoudnessNRTFactory.cpp
--- a/Engine/Plugins/Runtime/AudioSynesthesia/Source/AudioSynesthesiaCore/Private/LoudnessNRTFactory.cpp
+++ b/Engine/Plugins/Runtime/AudioSynesthesia/Source/AudioSynesthesiaCore/Private/LoudnessNRTFactory.cpp
@@ -1,6 +1,9 @@
 // Copyright Epic Games, Inc. All Rights Reserved.
 
 #include "LoudnessNRTFactory.h"
+
+#include "Algo/IsSorted.h"
+#include "AudioSynesthesiaCoreLog.h"
 #include "AudioSynesthesiaCustomVersion.h"
 #include "DSP/SlidingWindow.h"
 
@@ -22,7 +25,7 @@ namespace Audio
 	}
 
 	// Use for overall loudness storage
-	const int32 FLoudnessNRTResult::ChannelIndexOverall = INDEX_NONE;
+	const int32 FLoudnessNRTResult::ChannelIndexOverall = -1;
 
 
 	FLoudnessNRTResult::FLoudnessNRTResult()
@@ -31,7 +34,6 @@ namespace Audio
 	{
 	}
 
-
 	void FLoudnessNRTResult::Serialize(FArchive& Archive)
 	{
 		Archive.UsingCustomVersion(FAudioSynesthesiaCustomVersion::GUID);
@@ -204,10 +206,11 @@ namespace Audio
 	void FLoudnessNRTWorker::AnalyzeWindow(TArrayView<const float> InWindow, FLoudnessNRTResult& OutResult)
 	{
 		TArray<float> ChannelPerceptualEnergy;
+		ChannelPerceptualEnergy.AddUninitialized(NumChannels);
 		float OverallPerceptualEnergy = 0.0f;
 
 		// Calculate perceptual energy
-		OverallPerceptualEnergy = Analyzer->CalculatePerceptualEnergy(InWindow, NumChannels, ChannelPerceptualEnergy);
+		OverallPerceptualEnergy = Analyzer->CalculatePerceptualEnergy(InWindow, NumChannels, TArrayView<float>(ChannelPerceptualEnergy));
 
 		// Convert to perceptual energy to loudness
 		TArray<float> ChannelLoudness;
diff --git a/Engine/Plugins/Runtime/AudioSynesthesia/Source/AudioSynesthesiaCore/Public/FindNearestByTimestamp.h b/Engine/Plugins/Runtime/AudioSynesthesia/Source/AudioSynesthesiaCore/Public/FindNearestByTimestamp.h
new file mode 100644
--- /dev/null
+++ b/Engine/Plugins/Runtime/AudioSynesthesia/Source/AudioSynesthesiaCore/Public/FindNearestByTimestamp.h
@@ -0,0 +1,51 @@
+// Copyright Epic Games, Inc. All Rights Reserved.
+
+#pragma once
+
+#include "Algo/BinarySearch.h"
+#include "Containers/ArrayView.h"
+#include "Math/UnrealMath.h"
+
+
+namespace Audio
+{
+	/** FindNearestByTimestamp
+	 *
+	 *  Returns a pointer to the element with the timestamp nearest to InTimestamp in absolute terms.. 
+	 *
+	 *  ElementType must be a type with a "float Timestamp;" member.
+	 *  InArray is a array of ElementTypes which are sorted in ascending order via the Timestamp member.
+	 *  InTimestamp is the query timestamp.
+	 */
+	template<typename ElementType>
+	ElementType* FindNearestByTimestamp(TArrayView<ElementType> InArray, float InTimestamp)
+	{
+		if (0 == InArray.Num())
+		{
+			return nullptr;
+		}
+		
+		// Get upper and lower array indices with FLoudnessData closest to the query timestamp.
+		int32 UpperIndex = Algo::UpperBoundBy(InArray, InTimestamp, [](const ElementType& Datum) { return Datum.Timestamp; });
+		int32 LowerIndex = UpperIndex - 1;
+
+		// Ensure indices are within range.
+		const int32 MaxIdx = InArray.Num() - 1;
+		LowerIndex = FMath::Clamp(LowerIndex, 0, MaxIdx);
+		UpperIndex = FMath::Clamp(UpperIndex, 0, MaxIdx);
+
+		ElementType& LowerDatum = InArray[LowerIndex];
+		ElementType& UpperDatum = InArray[UpperIndex];
+
+		if (FMath::Abs(UpperDatum.Timestamp - InTimestamp) < FMath::Abs(LowerDatum.Timestamp - InTimestamp))
+		{
+			return &UpperDatum;
+		}
+		else
+		{
+			return &LowerDatum;
+		}
+	}
+}
+
+
diff --git a/Engine/Plugins/Runtime/AudioSynesthesia/Source/AudioSynesthesiaCore/Public/LKFSNRTFactory.h b/Engine/Plugins/Runtime/AudioSynesthesia/Source/AudioSynesthesiaCore/Public/LKFSNRTFactory.h
new file mode 100644
--- /dev/null
+++ b/Engine/Plugins/Runtime/AudioSynesthesia/Source/AudioSynesthesiaCore/Public/LKFSNRTFactory.h
@@ -0,0 +1,242 @@
+// Copyright Epic Games, Inc. All Rights Reserved.
+
+#pragma once
+
+#include "Containers/ArrayView.h"
+#include "Containers/StridedView.h"
+#include "DSP/SlidingWindow.h"
+#include "LoudnessAnalyzer.h"
+#include "IAudioAnalyzerNRTInterface.h"
+
+#include "LKFSNRTFactory.generated.h"
+
+namespace Audio { template <typename InSampleType> class TSlidingBuffer; }
+
+/** Data representing the instantaneous loudness of a audio. */
+USTRUCT(BlueprintType)
+struct FLKFSNRTResults
+{
+	GENERATED_BODY()
+
+	/** The audio channel index which produced this result. If -1, then
+	 * the data was produced by combining all other audio channels */
+	UPROPERTY(EditAnywhere, BlueprintReadOnly, Category = AudioAnalyzer)
+	int32 Channel = 0;
+
+	/** Time in seconds of the source audio which corresponds to the loudness measurements. */
+	UPROPERTY(EditAnywhere, BlueprintReadOnly, Category = AudioAnalyzer)
+	float Timestamp = 0.f;
+
+	/** The instantaneous, perceptually weighted energy. */
+	UPROPERTY(EditAnywhere, BlueprintReadOnly, Category = AudioAnalyzer)
+	float Energy = 0.f;
+
+	/** The instantaneous, perceptually weighted loudness. */
+	UPROPERTY(EditAnywhere, BlueprintReadOnly, Category = AudioAnalyzer)
+	float Loudness = 0.f;
+
+	/** The average loudness over a short term window of time. */
+	UPROPERTY(EditAnywhere, BlueprintReadOnly, Category = AudioAnalyzer)
+	float ShortTermLoudness = 0.f;
+};
+
+	/** The loudness data associated with an entire file. */
+USTRUCT(BlueprintType)
+struct FLKFSNRTAggregateStats
+{
+	GENERATED_BODY()
+
+	static constexpr float InvalidLoudness = TNumericLimits<float>::Min();
+	
+	/** The average loudness over the entire sound. */
+	UPROPERTY(EditAnywhere, BlueprintReadOnly, Category = AudioAnalyzer)
+	float IntegratedLoudness = InvalidLoudness;
+
+	/** The gated loudness aggregated over the entire sound. */
+	UPROPERTY(EditAnywhere, BlueprintReadOnly, Category = AudioAnalyzer)
+	float GatedLoudness = InvalidLoudness;
+};
+
+USTRUCT()
+struct FLKFSNRTChannelData
+{
+	GENERATED_BODY()
+
+	UPROPERTY()
+	FLKFSNRTAggregateStats AggregateStats;
+
+	UPROPERTY()
+	TArray<FLKFSNRTResults> LoudnessArray;
+
+	UPROPERTY()
+	FFloatInterval LoudnessInterval;
+};
+
+USTRUCT()
+struct FLKFSNRTWaveData
+{
+	GENERATED_BODY()
+
+	UPROPERTY()
+	bool bIsSortedChronologically = false;
+
+	UPROPERTY()
+	float DurationInSeconds = 0.f;
+
+	UPROPERTY()
+	TMap<int32, FLKFSNRTChannelData> ChannelData;
+};
+
+namespace Audio
+{
+	/**
+	 * Contains settings for loudness analyzer.
+	 */
+	class FLKFSNRTSettings : public IAnalyzerNRTSettings 
+	{
+	public:
+		static constexpr float DefaultShortTermLoudnessDuration = 3.f;
+
+		FLoudnessAnalyzerSettings LoudnessAnalyzerSettings = { 
+			.FFTSize = 2048,
+			.WindowType = EWindowType::Blackman,
+			.LoudnessCurveType = ELoudnessCurveType::K,
+			.MinAnalysisFrequency = 20.f,
+			.MaxAnalysisFrequency = 20000.f,
+			.ScalingMethod = ELoudnessAnalyzerScalingMethod::Corrected
+		};
+
+		/** Number of seconds between loudness measurements */
+		float AnalysisPeriod = 0.1f;
+
+		/** Number of seconds of audio analyzed for each loudness measurements */
+		float AnalysisWindowDuration = 0.4f;
+
+		/** Duration of audio analyzed for short term loudness. */
+		float ShortTermLoudnessDuration = DefaultShortTermLoudnessDuration;
+	};
+
+
+	/** FLKFSNRTResult contains the temporal evolution of loudness. */
+	class FLKFSNRTResult : public IAnalyzerNRTResult
+	{
+	public:
+		/** Denotes the overall loudness channel index as opposed individual channel indices. */
+		AUDIOSYNESTHESIACORE_API static const int32 ChannelIndexOverall;
+
+		AUDIOSYNESTHESIACORE_API FLKFSNRTResult();
+
+		/** Defines how to serialize result. */
+		AUDIOSYNESTHESIACORE_API virtual void Serialize(FArchive& Archive) override;
+
+		/** Initialize a channel if it does not already exist. */
+		AUDIOSYNESTHESIACORE_API void AddChannel(int32 InChannelIndex);
+
+		/** Appends an FLKFSNRTResults to the container. */
+		AUDIOSYNESTHESIACORE_API void Add(const FLKFSNRTResults& InDatum);
+
+		/** Returns true if this object data for the given channel index */
+		AUDIOSYNESTHESIACORE_API bool ContainsChannel(int32 InChannelIndex) const;
+
+		/** Returns const reference to FLKFSNRTResults array for individual channel. */
+		AUDIOSYNESTHESIACORE_API TArrayView<const FLKFSNRTResults> GetChannelLoudnessArray(int32 ChannelIdx) const;
+
+		/** Returns array view to FLKFSNRTResults array for individual channel. If
+		 * a caller modifies the timestamps of data, they may need to resort the
+		 * data by calling SortChronologically().
+		 */
+		AUDIOSYNESTHESIACORE_API TArrayView<FLKFSNRTResults> GetMutableChannelLoudnessArray(int32 ChannelIdx); 
+
+		/** Returns const reference to FLKFSNRTResults array associated with overall loudness. */
+		AUDIOSYNESTHESIACORE_API TArrayView<const FLKFSNRTResults> GetLoudnessArray() const;
+
+		/** Return the aggregate stats associated with a specific channel */
+		AUDIOSYNESTHESIACORE_API const FLKFSNRTAggregateStats* FindAggregateLoudnessStats(int32 InChannelIdx) const;
+
+		/** Sets the aggregated loudness stats for a given channel */
+		AUDIOSYNESTHESIACORE_API void SetAggregateLoudnessStats(int32 InChannelIdx, FLKFSNRTAggregateStats InStats);
+
+		/** Returns the channel indices available in result. */
+		AUDIOSYNESTHESIACORE_API void GetChannels(TArray<int32>& OutChannels) const;
+
+		/** Gets the duration of the analyzed audio. */
+		AUDIOSYNESTHESIACORE_API virtual float GetDurationInSeconds() const override;
+
+		/** Sets the duration of the analyzed audio. */
+		AUDIOSYNESTHESIACORE_API void SetDurationInSeconds(float InDuration);
+
+		/** Returns true if FLKFSNRTResults arrays are sorted in chronologically ascending order via their timestamp. */
+		AUDIOSYNESTHESIACORE_API bool IsSortedChronologically() const;
+
+		/** Sorts FLKFSNRTResults arrays in chronologically ascending order via their timestamp. */
+		AUDIOSYNESTHESIACORE_API void SortChronologically();
+
+	private:
+		FLKFSNRTWaveData Data;
+	};
+
+	/** 
+	 * FLKFSNRTWorker performs loudness analysis on input sample buffers.
+	 */
+	class FLKFSNRTWorker : public IAnalyzerNRTWorker
+	{
+	public:
+		/** Construct a worker */
+		AUDIOSYNESTHESIACORE_API FLKFSNRTWorker(const FAnalyzerNRTParameters& InParams, const FLKFSNRTSettings& InAnalyzerSettings);
+
+		/**
+		 * Analyzes input sample buffer and updates result. 
+		 */
+		AUDIOSYNESTHESIACORE_API virtual void Analyze(TArrayView<const float> InAudio, IAnalyzerNRTResult* OutResult) override;
+
+		/**
+		 * Call when all audio data has been analyzed. 
+		 */
+		AUDIOSYNESTHESIACORE_API virtual void Finalize(IAnalyzerNRTResult* OutResult) override;
+
+	private:
+
+
+		void TuneSlidingWindwoHopSize(const int32 InAnalysisHopFrames, const int32 InWindowFrames, int32& OutBestWindowHopFrames, int32& OutBestAnalysisHopWindows) const;
+		void AnalyzeSlidingWindow(TAutoSlidingWindow<float>& InSlidingWindow, FLKFSNRTResult& OutResult);
+		void AnalyzeWindow(TArrayView<const float> InWindow, FLKFSNRTResult& OutResult);
+
+		void AddCurrentLoudnessDataToResult(FLKFSNRTResult& OutResult);
+		void CalculateAggregateLoudnessStats(FLKFSNRTResult& InOutResult) const;
+		void FinalizeChannels(FLKFSNRTResult& OutResult);
+
+		int32 NumChannels = 0;
+		int32 NumAnalyzedBuffers = 0;
+		int32 NumSlidingWindowHopFrames = 0;
+		int32 NumFrames = 0;
+		float SampleRate = 0.f;
+		float ShortTermLoudnessDuration = 0.f;
+		int32 NumAnalysisHopWindows = 1;
+		int32 NumAnalysisWindows = 1;
+		TArray<float> ChannelPerceptualEnergy;
+		TArray<float> OverallPerceptualEnergy;
+		TArray<float> InternalWindow;
+		TUniquePtr<TSlidingBuffer<float>> InternalBuffer;
+		TUniquePtr<FMultichannelLoudnessAnalyzer> Analyzer;
+	};
+
+	/**
+	 * Defines the LKFSNRT analyzer and creates related classes.
+	 */
+	class FLKFSNRTFactory : public IAnalyzerNRTFactory
+	{
+	public:
+
+		/** Name of specific analyzer type. */
+		AUDIOSYNESTHESIACORE_API virtual FName GetName() const override;
+
+		/** Human readable name of analyzer. */
+		AUDIOSYNESTHESIACORE_API virtual FString GetTitle() const override;
+
+		/** Creates a new FLKFSNRTResult */
+		AUDIOSYNESTHESIACORE_API virtual TUniquePtr<IAnalyzerNRTResult> NewResult() const override;
+
+		/** Creates a new FLKFSNRTWorker. This expects IAnalyzerNRTSettings to be a valid pointer to a FLKFSNRTSettings object. */
+		AUDIOSYNESTHESIACORE_API virtual TUniquePtr<IAnalyzerNRTWorker> NewWorker(const FAnalyzerNRTParameters& InParams, const IAnalyzerNRTSettings* InSettings) const override;
+	};
+}
diff --git a/Engine/Plugins/Runtime/AudioSynesthesia/Source/AudioSynesthesiaCore/Public/LoudnessAnalyzer.h b/Engine/Plugins/Runtime/AudioSynesthesia/Source/AudioSynesthesiaCore/Public/LoudnessAnalyzer.h
--- a/Engine/Plugins/Runtime/AudioSynesthesia/Source/AudioSynesthesiaCore/Public/LoudnessAnalyzer.h
+++ b/Engine/Plugins/Runtime/AudioSynesthesia/Source/AudioSynesthesiaCore/Public/LoudnessAnalyzer.h
@@ -174,8 +174,23 @@ namespace Audio
 			 *
 			 * The overall perceptual energy is returned.
 			 */
+			UE_DEPRECATED(5.7, "Use method which accepts a TArrayView<float> for OutChannelEnergies")
             UE_API float CalculatePerceptualEnergy(TArrayView<const float> InSampleView, const int32 InNumChannels, TArray<float>& OutChannelEnergies); 
 
+			/** 
+			 * Calculate perceptual energy. 
+			 *
+			 * InSampleView		  InSampleView is expected to contain interleaved audio with 
+			 *                    (InSettings.FFTSize * InNumChannels) number of samples.
+			 *
+			 * InNumChannels      Number of input channels.
+			 *
+			 * OutChannelEnergies Will be populated with energies on a per channel basis.
+			 *
+			 * The overall perceptual energy is returned.
+			 */
+            UE_API float CalculatePerceptualEnergy(TArrayView<const float> InSampleView, const int32 InNumChannels, TArrayView<float> OutChannelEnergies); 
+
 			/** 
 			 * Calculate loudness
 			 *
@@ -188,14 +203,30 @@ namespace Audio
 			 *
 			 * The overall loudness is returned.
 			 */
+			UE_DEPRECATED(5.7, "Use method which accepts a TArrayView<float> for OutChannelLoudness")
             UE_API float CalculateLoudness(TArrayView<const float> InSampleView, const int32 InNumChannels, TArray<float>& OutChannelLoudness);
 
+			/** 
+			 * Calculate loudness
+			 *
+			 * InSampleView		  InSampleView is expected to contain interleaved audio with 
+			 *                    (InSettings.FFTSize * InNumChannels) number of samples.
+			 *
+			 * InNumChannels      Number of input channels.
+			 *
+			 * OutChannelLoudness Will be populated with energies on a per channel basis.
+			 *
+			 * The overall loudness is returned.
+			 */
+            UE_API float CalculateLoudness(TArrayView<const float> InSampleView, const int32 InNumChannels, TArrayView<float> OutChannelLoudness);
+
 			/**
 			 * Return const reference to settings used inside this analyzer.
 			 */
             UE_API const FLoudnessAnalyzerSettings& GetSettings() const;
 
         private:
+			void InitializeChannelWeights(int32 InNumChannels);
 
             FAlignedFloatBuffer MonoBuffer;
             TArray<float> ChannelWeights;
