Adding LKFS realtime analyzer to audio synesthesia

diff --git a/Engine/Plugins/Runtime/AudioSynesthesia/Source/AudioSynesthesia/Classes/LKFS.h b/Engine/Plugins/Runtime/AudioSynesthesia/Source/AudioSynesthesia/Classes/LKFS.h
new file mode 100644
--- /dev/null
+++ b/Engine/Plugins/Runtime/AudioSynesthesia/Source/AudioSynesthesia/Classes/LKFS.h
@@ -0,0 +1,134 @@
+// Copyright Epic Games, Inc. All Rights Reserved.
+
+#pragma once
+
+#include "AudioSynesthesia.h"
+#include "LKFSFactory.h"
+
+#include "LKFS.generated.h"
+
+#define UE_API AUDIOSYNESTHESIA_API
+
+/** ULKFSSettings
+ *
+ * Settings for a ULKFSAnalyzer.
+ */
+UCLASS(MinimalAPI, Blueprintable)
+class ULKFSSettings : public UAudioSynesthesiaSettings
+{
+	GENERATED_BODY()
+public:
+
+	ULKFSSettings() {}
+
+	/** Number of seconds between loudness measurements. The standard is 0.1 seconds. */
+	UPROPERTY(EditAnywhere, BlueprintReadOnly, Category = AudioAnalyzer, meta = (ClampMin = "0.01", ClampMax = "1.00"))
+	float AnalysisPeriod = 0.1f;
+
+	/** Number of seconds of audio analyzed for each loudness measurements. The standard is 0.4 seconds. */
+	UPROPERTY(EditAnywhere, BlueprintReadOnly, Category=AudioAnalyzer, meta = (ClampMin = "0.10", ClampMax = "1.00"))
+	float AnalysisWindowDuration = 0.4f;
+
+	/** Duration of audio analyzed for short term loudness. The standard is 3.0 seconds. */
+	UPROPERTY(EditAnywhere, BlueprintReadOnly, Category=AudioAnalyzer, meta = (ClampMin = "1.00", ClampMax = "30.00"))
+	float ShortTermLoudnessDuration = Audio::FLKFSSettings::StandardShortTermLoudnessDuration;
+
+	/** Number of seconds between long term (integrated and gated) loudness updates.  The standard is 1.0 seconds. */
+	UPROPERTY(EditAnywhere, BlueprintReadOnly, Category=AudioAnalyzer, meta = (ClampMin = "0.1", ClampMax = "10.00"))
+	float IntegratedLoudnessAnalysisPeriod = Audio::FLKFSSettings::StandardIntegratedLoudnessAnalysisPeriod;
+
+	/** Duration of audio analyzed for long term (integrated and gated) loudness. The real time LKFS analyzer
+	 * performs integrated loudness duration over a large window instead of performing calculations over
+	 * the entire audio stream. */
+	UPROPERTY(EditAnywhere, BlueprintReadOnly, Category=AudioAnalyzer, meta = (ClampMin = "10.00", ClampMax = "3600.00"))
+	float IntegratedLoudnessDuration = Audio::FLKFSSettings::DefaultIntegratedLoudnessDuration;
+
+	/** Convert ULKFSSettings to IAnalyzerSettings */
+	UE_API TUniquePtr<Audio::IAnalyzerSettings> GetSettings(const int32 InSampleRate, const int32 InNumChannels) const;
+
+#if WITH_EDITOR
+	UE_API virtual FText GetAssetActionName() const override;
+
+	UE_API virtual UClass* GetSupportedClass() const override;
+#endif
+};
+
+/** Delegate to receive all overall loudness results (time-stamped in an array) since last delegate call. */
+DECLARE_DYNAMIC_MULTICAST_DELEGATE_OneParam(FOnOverallLKFSResults, const TArray<FLKFSResults>&, LKFSResults);
+
+/** shadow delegate declaration for above */
+DECLARE_MULTICAST_DELEGATE_TwoParams(FOnOverallLKFSResultsNative, ULKFSAnalyzer*, TArrayView<const FLKFSResults>);
+
+/** Delegate to receive only the most recent overall meter results. */
+DECLARE_DYNAMIC_MULTICAST_DELEGATE_OneParam(FOnLatestOverallLKFSResults, const FLKFSResults&, LatestOverallLKFSResults);
+
+/** shadow delegate declaration for above */
+DECLARE_MULTICAST_DELEGATE_TwoParams(FOnLatestOverallLKFSResultsNative, ULKFSAnalyzer*, const FLKFSResults&);
+
+/** Delegate to receive all meter results per channel (time-stamped in an array) since last delegate call. */
+DECLARE_DYNAMIC_MULTICAST_DELEGATE_TwoParams(FOnPerChannelLKFSResults, int32, ChannelIndex, const TArray<FLKFSResults>&, LKFSResults);
+
+/** shadow delegate declaration for above */
+DECLARE_MULTICAST_DELEGATE_ThreeParams(FOnPerChannelLKFSResultsNative, ULKFSAnalyzer*, int32, TArrayView<const FLKFSResults>);
+
+/** Delegate to receive only the most recent overall meter result per channel. */
+DECLARE_DYNAMIC_MULTICAST_DELEGATE_TwoParams(FOnLatestPerChannelLKFSResults, int32, ChannelIndex, const FLKFSResults&, LatestLKFSResults);
+
+/** shadow delegate declaration for above */
+DECLARE_MULTICAST_DELEGATE_ThreeParams(FOnLatestPerChannelLKFSResultsNative, ULKFSAnalyzer*, int32, const FLKFSResults&);
+
+/** ULKFSAnalyzer
+ *
+ * ULKFSAnalyzer calculates the current amplitude of an
+ * audio bus in real-time.
+ */
+UCLASS(MinimalAPI, Blueprintable)
+class ULKFSAnalyzer : public UAudioAnalyzer
+{
+	GENERATED_BODY()
+
+public:
+
+	UE_API ULKFSAnalyzer();
+
+	/** The settings for the meter audio analyzer.  */
+	UPROPERTY(EditAnywhere, BlueprintReadWrite, Category = AudioAnalyzer)
+	TObjectPtr<ULKFSSettings> Settings;
+
+	/** Delegate to receive all overall meter results since last delegate call. */
+	UPROPERTY(BlueprintAssignable)
+	FOnOverallLKFSResults OnOverallLKFSResults;
+
+	FOnOverallLKFSResultsNative OnOverallLKFSResultsNative;
+
+	/** Delegate to receive all meter results, per-channel, since last delegate call. */
+	UPROPERTY(BlueprintAssignable)
+	FOnPerChannelLKFSResults OnPerChannelLKFSResults;
+
+	FOnPerChannelLKFSResultsNative OnPerChannelLKFSResultsNative;
+
+	/** Delegate to receive the latest overall meter results. */
+	UPROPERTY(BlueprintAssignable)
+	FOnLatestOverallLKFSResults OnLatestOverallLKFSResults;
+
+	FOnLatestOverallLKFSResultsNative OnLatestOverallLKFSResultsNative;
+
+	/** Delegate to receive the latest per-channel meter results. */
+	UPROPERTY(BlueprintAssignable)
+	FOnLatestPerChannelLKFSResults OnLatestPerChannelLKFSResults;
+
+	FOnLatestPerChannelLKFSResultsNative OnLatestPerChannelLKFSResultsNative;
+
+	/** Convert ULKFSSettings to IAnalyzerSettings */
+	UE_API TUniquePtr<Audio::IAnalyzerSettings> GetSettings(const int32 InSampleRate, const int32 InNumChannels) const override;
+
+	/** Broadcasts results to any delegates if hooked up. */
+	UE_API void BroadcastResults() override;
+
+protected:
+
+	/** Return the name of the IAudioAnalyzerFactory associated with this UAudioAnalyzer */
+	UE_API FName GetAnalyzerFactoryName() const override;
+};
+
+#undef UE_API
diff --git a/Engine/Plugins/Runtime/AudioSynesthesia/Source/AudioSynesthesia/Private/LKFS.cpp b/Engine/Plugins/Runtime/AudioSynesthesia/Source/AudioSynesthesia/Private/LKFS.cpp
new file mode 100644
--- /dev/null
+++ b/Engine/Plugins/Runtime/AudioSynesthesia/Source/AudioSynesthesia/Private/LKFS.cpp
@@ -0,0 +1,106 @@
+// Copyright Epic Games, Inc. All Rights Reserved.
+
+#include "LKFS.h"
+#include "LKFSFactory.h"
+
+#include UE_INLINE_GENERATED_CPP_BY_NAME(LKFS)
+
+TUniquePtr<Audio::IAnalyzerSettings> ULKFSSettings::GetSettings(const int32 InSampleRate, const int32 InNumChannels) const
+{
+	TUniquePtr<Audio::FLKFSSettings> Settings = MakeUnique<Audio::FLKFSSettings>();
+
+	Settings->AnalysisPeriod = AnalysisPeriod;
+	Settings->AnalysisWindowDuration = AnalysisWindowDuration;
+	Settings->ShortTermLoudnessDuration = ShortTermLoudnessDuration;
+	Settings->IntegratedLoudnessAnalysisPeriod = IntegratedLoudnessAnalysisPeriod;
+	Settings->IntegratedLoudnessDuration = IntegratedLoudnessDuration;
+
+	return Settings;
+}
+
+#if WITH_EDITOR
+FText ULKFSSettings::GetAssetActionName() const
+{
+	return NSLOCTEXT("AssetTypeActions", "AssetTypeActions_AssetSoundSynesthesiaLKFSSettings", "Synesthesia Real-Time Settings (LKFS)");
+}
+
+UClass* ULKFSSettings::GetSupportedClass() const
+{
+	return ULKFSSettings::StaticClass();
+}
+#endif
+
+ULKFSAnalyzer::ULKFSAnalyzer()
+{
+	Settings = CreateDefaultSubobject<ULKFSSettings>(TEXT("DefaultLKFSSettings"));
+}
+
+TUniquePtr<Audio::IAnalyzerSettings> ULKFSAnalyzer::GetSettings(const int32 InSampleRate, const int32 InNumChannels) const
+{
+	TUniquePtr<Audio::IAnalyzerSettings> AnalyzerSettings;
+
+	if (Settings)
+	{
+		AnalyzerSettings = Settings->GetSettings(InSampleRate, InNumChannels);
+	}
+
+	return AnalyzerSettings;
+}
+
+void ULKFSAnalyzer::BroadcastResults()
+{
+	TUniquePtr<const Audio::FLKFSAnalyzerResult> AnalyzerResults = GetResults<Audio::FLKFSAnalyzerResult>();
+	if (!AnalyzerResults.IsValid())
+	{
+		return;
+	}
+
+	int32 NumChannels = AnalyzerResults->GetNumChannels();
+
+	if (NumChannels > 0)
+	{
+		// Broadcast overall results
+		bool bIsOnOverallLKFSResultsBound = OnOverallLKFSResults.IsBound() || OnOverallLKFSResultsNative.IsBound();
+		bool bIsOnLatestOverallLKFSResultsBound = OnLatestOverallLKFSResults.IsBound() || OnLatestOverallLKFSResultsNative.IsBound();
+		if (bIsOnOverallLKFSResultsBound || bIsOnLatestOverallLKFSResultsBound)
+		{
+			const TArray<FLKFSResults>& OverallLKFSArray = AnalyzerResults->GetLoudnessResults();
+			if (OverallLKFSArray.Num() > 0)
+			{
+				OnOverallLKFSResults.Broadcast(OverallLKFSArray);
+				OnOverallLKFSResultsNative.Broadcast(this, OverallLKFSArray);
+
+				const FLKFSResults& Latest = OverallLKFSArray[OverallLKFSArray.Num() - 1];
+				OnLatestOverallLKFSResults.Broadcast(Latest);
+				OnLatestOverallLKFSResultsNative.Broadcast(this, Latest);
+			}
+		}
+
+		// Broadcast per channel results
+		bool bIsOnPerChannelLKFSResultsBound = OnPerChannelLKFSResults.IsBound() || OnPerChannelLKFSResultsNative.IsBound();
+		bool bIsOnLatestPerChannelLKFSResultsBound = OnLatestPerChannelLKFSResults.IsBound() || OnLatestPerChannelLKFSResultsNative.IsBound();
+
+		if (bIsOnPerChannelLKFSResultsBound || bIsOnLatestPerChannelLKFSResultsBound)
+		{
+			for (int32 ChannelIndex = 0; ChannelIndex < NumChannels; ++ChannelIndex)
+			{
+				const TArray<FLKFSResults>& LKFSArray = AnalyzerResults->GetChannelLoudnessResults(ChannelIndex);
+				if (LKFSArray.Num() > 0)
+				{
+					OnPerChannelLKFSResults.Broadcast(ChannelIndex, LKFSArray);
+					OnPerChannelLKFSResultsNative.Broadcast(this, ChannelIndex, LKFSArray);
+
+					const FLKFSResults& Latest = LKFSArray[LKFSArray.Num() - 1];
+					OnLatestPerChannelLKFSResults.Broadcast(ChannelIndex, Latest);
+					OnLatestPerChannelLKFSResultsNative.Broadcast(this, ChannelIndex, Latest);
+				}
+			}
+		}
+	}
+}
+
+FName ULKFSAnalyzer::GetAnalyzerFactoryName() const
+{
+	static const FName FactoryName(TEXT("LKFSFactory"));
+	return FactoryName;
+}
diff --git a/Engine/Plugins/Runtime/AudioSynesthesia/Source/AudioSynesthesiaCore/Private/AudioSynesthesiaCoreModule.cpp b/Engine/Plugins/Runtime/AudioSynesthesia/Source/AudioSynesthesiaCore/Private/AudioSynesthesiaCoreModule.cpp
--- a/Engine/Plugins/Runtime/AudioSynesthesia/Source/AudioSynesthesiaCore/Private/AudioSynesthesiaCoreModule.cpp
+++ b/Engine/Plugins/Runtime/AudioSynesthesia/Source/AudioSynesthesiaCore/Private/AudioSynesthesiaCoreModule.cpp
@@ -8,6 +8,7 @@
 #include "HAL/LowLevelMemTracker.h"
 #include "LoudnessFactory.h"
 #include "LoudnessNRTFactory.h"
+#include "LKFSFactory.h"
 #include "LKFSNRTFactory.h"
 #include "MeterFactory.h"
 #include "Modules/ModuleManager.h"
@@ -35,6 +36,7 @@ namespace Audio
 				IModularFeatures::Get().RegisterModularFeature(FLoudnessFactory::GetModularFeatureName(), &LoudnessFactory);
 				IModularFeatures::Get().RegisterModularFeature(FConstantQFactory::GetModularFeatureName(), &ConstantQFactory);
 				IModularFeatures::Get().RegisterModularFeature(FMeterFactory::GetModularFeatureName(), &MeterFactory);
+				IModularFeatures::Get().RegisterModularFeature(FLKFSFactory::GetModularFeatureName(), &LKFSFactory);
 				IModularFeatures::Get().RegisterModularFeature(FSynesthesiaSpectrumAnalysisFactory::GetModularFeatureName(), &SpectralAnalysisFactory);
 			}
 
@@ -50,6 +52,7 @@ namespace Audio
 				IModularFeatures::Get().UnregisterModularFeature(FLoudnessFactory::GetModularFeatureName(), &LoudnessFactory);
 				IModularFeatures::Get().UnregisterModularFeature(FConstantQFactory::GetModularFeatureName(), &ConstantQFactory);
 				IModularFeatures::Get().UnregisterModularFeature(FMeterFactory::GetModularFeatureName(), &MeterFactory);
+				IModularFeatures::Get().UnregisterModularFeature(FLKFSFactory::GetModularFeatureName(), &LKFSFactory);
 				IModularFeatures::Get().UnregisterModularFeature(FSynesthesiaSpectrumAnalysisFactory::GetModularFeatureName(), &SpectralAnalysisFactory);
 			}
 
@@ -62,6 +65,7 @@ namespace Audio
 			FLoudnessFactory LoudnessFactory;
 			FConstantQFactory ConstantQFactory;
 			FMeterFactory MeterFactory; 
+			FLKFSFactory LKFSFactory; 
 			FSynesthesiaSpectrumAnalysisFactory SpectralAnalysisFactory;
 	};
 }
diff --git a/Engine/Plugins/Runtime/AudioSynesthesia/Source/AudioSynesthesiaCore/Private/LKFSFactory.cpp b/Engine/Plugins/Runtime/AudioSynesthesia/Source/AudioSynesthesiaCore/Private/LKFSFactory.cpp
new file mode 100644
--- /dev/null
+++ b/Engine/Plugins/Runtime/AudioSynesthesia/Source/AudioSynesthesiaCore/Private/LKFSFactory.cpp
@@ -0,0 +1,374 @@
+// Copyright Epic Games, Inc. All Rights Reserved.
+
+#include "LKFSFactory.h"
+
+#include "Algo/IsSorted.h"
+#include "AudioSynesthesiaCoreLog.h"
+#include "Containers/Array.h"
+#include "DSP/SlidingWindow.h"
+#include "LKFSPrivate.h"
+
+namespace Audio
+{
+	namespace LKFSFactoryPrivate
+	{
+		float LKFSToEnergy(float InLKFS)
+		{
+			return FMath::Pow(10, (InLKFS + 0.691f) / 10.f);
+		}
+
+		// Calculates average value from a circular buffer. Allows for interleaved data
+		// by providing a "stride". 
+		float CalculateAverageFromHistory(const TCircularBuffer<float>& InHistory, int32 InStartIndex, int32 InEndIndex, int32 InStride)
+		{
+			float Accum = 0.f;
+			for (int32 Index = InStartIndex; Index < InEndIndex; Index += InStride)
+			{
+				Accum += InHistory[Index];
+			}
+			int32 Num = (InEndIndex - InStartIndex) / InStride;
+			if (Num > 0)
+			{
+				return Accum /= Num;
+			}
+			return 0.f;
+		}
+
+		// Calculates gated average value from a circular buffer. Allows for interleaved data
+		// by providing a "stride". 
+		float CalculateGatedAverageFromHistory(const TCircularBuffer<float>& InHistory, int32 InStartIndex, int32 InEndIndex, int32 InStride, float InGate)
+		{
+			float Accum = 0.f;
+			int32 Num = 0;
+			for (int32 Index = InStartIndex; Index < InEndIndex; Index += InStride)
+			{
+				float Val = InHistory[Index];
+				if (Val > InGate)
+				{
+					Accum += Val;
+					Num++;
+				}
+			}
+			if (Num > 0)
+			{
+				return Accum /= Num;
+			}
+			return InGate;
+		}
+
+		float CalculateGatedLoudnessFromHistory(const TCircularBuffer<float>& InHistory, int32 InStartIndex, int32 InEndIndex, int32 InStride)
+		{
+			// Gated Loudness only examines loudnesses above a certain gating loudness. 
+			//
+			// It is a 2 pass algorithm 
+			// Pass 1: Use a set loudness threshold of -70lkfs to determine an initial gated loudness
+			// Pass 2: Use the initial gated loudness to derive a new loudness threshold and calculate the gate loudness again. 
+			constexpr float InitLoudnessThreshold = -70.f;
+			float InitEnergyThreshold = LKFSToEnergy(InitLoudnessThreshold);
+
+			const float GatedEnergyPass1 = CalculateGatedAverageFromHistory(InHistory, InStartIndex, InEndIndex, InStride, InitEnergyThreshold);
+
+			const float NewThreshold = GatedEnergyPass1 * 0.1f;
+			const float GatedEnergyPass2 = CalculateGatedAverageFromHistory(InHistory, InStartIndex, InEndIndex, InStride, NewThreshold);
+			return FLoudnessAnalyzer::ConvertPerceptualEnergyToLKFS(GatedEnergyPass2);
+		}
+	}
+}
+
+namespace Audio
+{
+	// Use for overall loudness storage
+	const int32 FLKFSAnalyzerResult::ChannelIndexOverall = -1;
+
+	FLKFSAnalyzerResult::FLKFSAnalyzerResult()
+	{
+	}
+
+	int32 FLKFSAnalyzerResult::GetNumChannels() const
+	{
+		int32 Num = ChannelResults.Num();
+		if (Num > 0)
+		{
+			return Num - 1; // Subtract one for storage of "Overall" channel
+		}
+		return 0;
+	}
+
+	void FLKFSAnalyzerResult::Add(FLKFSResults InDatum)
+	{
+		TArray<FLKFSResults>& ChannelData = ChannelResults.FindOrAdd(InDatum.Channel);
+		ChannelData.Add(MoveTemp(InDatum));
+	}
+
+	const TArray<FLKFSResults>& FLKFSAnalyzerResult::GetChannelLoudnessResults(int32 ChannelIdx) const
+	{
+		if (const TArray<FLKFSResults>* ChannelData = ChannelResults.Find(ChannelIdx))
+		{
+			return *ChannelData;
+		}
+		else
+		{
+			static const TArray<FLKFSResults> EmptyArray;
+			return EmptyArray;
+		}
+	}
+
+	const TArray<FLKFSResults>& FLKFSAnalyzerResult::GetLoudnessResults() const
+	{
+		return GetChannelLoudnessResults(ChannelIndexOverall);
+	}
+
+	FLKFSWorker::FLKFSWorker(const FAnalyzerParameters& InParams, const FLKFSSettings& InAnalyzerSettings)
+	: NumChannels(InParams.NumChannels)
+	, SampleRate(InParams.SampleRate)
+	, ChannelEnergyHistory(1)
+	, OverallEnergyHistory(1)
+	, Analyzer(InParams.SampleRate, InAnalyzerSettings.LoudnessAnalyzerSettings)
+	{
+		check(NumChannels > 0);
+		check(SampleRate > 0.f);
+		check(InAnalyzerSettings.ShortTermLoudnessDuration > 0.f);
+		check(InAnalyzerSettings.IntegratedLoudnessDuration > 0.f);
+		check(InAnalyzerSettings.IntegratedLoudnessAnalysisPeriod> 0.f);
+
+
+		int32 NumSlidingWindowFrames = Analyzer.GetSettings().WindowSize;
+
+		int32 NumAnalysisFrames = FMath::CeilToInt(InAnalyzerSettings.AnalysisWindowDuration * InParams.SampleRate);
+		int32 NumAnalysisHopFrames = FMath::CeilToInt(InAnalyzerSettings.AnalysisPeriod * InParams.SampleRate);
+
+		// There may be multiple calls to the LoudnessAnalyzer to produce a single 
+		// FLoudnessDatum. The sliding window hop is tuned here so that it best 
+		// matches the desired AnalysisPeriod while maintaining between a 25% to 
+		// 75% window overlap.
+		LKFSPrivate::TuneSlidingWindwoHopSize(NumAnalysisHopFrames, NumSlidingWindowFrames, NumSlidingWindowHopFrames, NumAnalysisHopWindows);
+
+		// Determine how many windows to analyze per loudness result
+		NumAnalysisWindows = FMath::Max(1, ((NumAnalysisFrames - (NumSlidingWindowFrames / 2)) / NumSlidingWindowHopFrames) + 1);
+
+		// Include NumChannels in calculating window size because windows are generated 
+		// with interleaved samples and deinterleaved later. 
+		int32 NumSlidingWindowSamples = NumSlidingWindowFrames * NumChannels;
+		int32 NumSlidingWindowHopSamples = NumSlidingWindowHopFrames * NumChannels;
+
+		InternalBuffer = TSlidingBuffer<float>(NumSlidingWindowSamples, NumSlidingWindowHopSamples);
+		InternalWindow.AddUninitialized(NumSlidingWindowSamples);
+		ChannelPerceptualEnergy.AddZeroed(NumChannels * NumAnalysisWindows);
+		OverallPerceptualEnergy.AddZeroed(NumAnalysisWindows);
+
+		// Initialize settings for calculating aggregate stats
+		const int32 NumResultHopFrames = NumAnalysisHopWindows * NumSlidingWindowHopFrames;
+		NumShortTermHistoryResults = FMath::Max(1, FMath::DivideAndRoundUp((int32) (InAnalyzerSettings.ShortTermLoudnessDuration * SampleRate), NumResultHopFrames));
+		NumIntegratedHistoryResults = FMath::Max(1, FMath::DivideAndRoundUp((int32) (InAnalyzerSettings.IntegratedLoudnessDuration * SampleRate), NumResultHopFrames));
+		NumIntegratedHopHistoryResults = FMath::Max(1, FMath::DivideAndRoundUp((int32) (InAnalyzerSettings.IntegratedLoudnessAnalysisPeriod * SampleRate), NumResultHopFrames));
+
+		const int32 MaxHistory = FMath::Max(NumIntegratedHistoryResults, NumShortTermHistoryResults);
+
+		ChannelEnergyHistory = TCircularBuffer<float>(MaxHistory * NumChannels);
+		OverallEnergyHistory = TCircularBuffer<float>(MaxHistory);
+
+		for (int32 ChannelIdx = 0; ChannelIdx < NumChannels; ChannelIdx++)
+		{
+			AggregateLoudnessStats.Add(ChannelIdx);
+		}
+		AggregateLoudnessStats.Add(FLKFSAnalyzerResult::ChannelIndexOverall);
+	}
+
+	void FLKFSWorker::Analyze(TArrayView<const float> InAudio, IAnalyzerResult* OutResult) 
+	{
+		// Assume that outer layers ensured that this is of correct type.
+		FLKFSAnalyzerResult* LoudnessResult = static_cast<FLKFSAnalyzerResult*>(OutResult);
+
+		check(nullptr != LoudnessResult);
+		
+		TAutoSlidingWindow<float> SlidingWindow(InternalBuffer, InAudio, InternalWindow);
+
+		AnalyzeSlidingWindow(SlidingWindow, *LoudnessResult);
+	}
+
+	void FLKFSWorker::AnalyzeSlidingWindow(TAutoSlidingWindow<float>& InSlidingWindow, FLKFSAnalyzerResult& OutResult)
+	{
+		check(NumAnalysisWindows > 0);
+		check(NumAnalysisHopWindows > 0);
+		
+		for (const TArray<float>& Window : InSlidingWindow)
+		{
+			AnalyzeWindow(Window);
+
+			if ((NumAnalyzedBuffers >= NumAnalysisWindows) && ((NumAnalyzedBuffers - NumAnalysisWindows) % NumAnalysisHopWindows) == 0)
+			{
+				AddCurrentLoudnessDataToResult(OutResult);
+			}
+		}
+	}
+
+	void FLKFSWorker::AnalyzeWindow(TArrayView<const float> InWindow)
+	{
+		const int32 WrappedWindowIdx = (NumAnalyzedBuffers % NumAnalysisWindows);
+		const int32 ChannelPerceptualEnergyWritePos = WrappedWindowIdx * NumChannels;
+
+		// Calculate perceptual energy
+		OverallPerceptualEnergy[WrappedWindowIdx] = Analyzer.CalculatePerceptualEnergy(InWindow, NumChannels, MakeArrayView<float>(ChannelPerceptualEnergy.GetData() + ChannelPerceptualEnergyWritePos, NumChannels));
+
+		// update counters
+		NumAnalyzedBuffers++;
+	}
+
+	void FLKFSWorker::AddCurrentLoudnessDataToHistory()
+	{
+		// Compute loudness per channel and store history
+		for (int32 ChannelIdx = 0; ChannelIdx < NumChannels; ChannelIdx++)
+		{
+			// Determine average perceptual energy over analysis window
+			float PerceptualEnergy = 0.f;
+			for (int32 Pos = ChannelIdx; Pos < ChannelPerceptualEnergy.Num(); Pos += NumChannels)
+			{
+				PerceptualEnergy += ChannelPerceptualEnergy[Pos];
+			}
+			PerceptualEnergy /= (float)NumAnalysisWindows;
+
+			// Store perceptual energy in history 
+			ChannelEnergyHistory[(NumChannels * NumEnergyHistory) + ChannelIdx] = PerceptualEnergy;
+		}
+
+		{
+			// Compute loudness overall
+			float PerceptualEnergy = 0.f;
+			for (int32 Pos = 0; Pos < OverallPerceptualEnergy.Num(); Pos++)
+			{
+				PerceptualEnergy += OverallPerceptualEnergy[Pos];
+			}
+			PerceptualEnergy /= (float)NumAnalysisWindows;
+
+			OverallEnergyHistory[NumEnergyHistory] = PerceptualEnergy;
+		}
+	
+		NumEnergyHistory++;
+	}
+
+	void FLKFSWorker::AddCurrentLoudnessDataToResult(FLKFSAnalyzerResult& OutResult)
+	{
+		AddCurrentLoudnessDataToHistory();
+
+		UpdateAggregateLoudnessStats();
+
+		// Accounting for all the frames positions analyzed in creating the 
+		// data, determine the center frame. 
+		const int32 CenterFramePos = ((2 * NumAnalyzedBuffers - NumAnalysisWindows) * NumSlidingWindowHopFrames + Analyzer.GetSettings().WindowSize) / 2;
+		const float Timestamp = CenterFramePos / SampleRate;
+
+
+		// Create results and store
+		for (int32 ChannelIdx = 0; ChannelIdx < NumChannels; ChannelIdx++)
+		{
+			int32 HistoryIndex = ((NumEnergyHistory - 1) * NumChannels) + ChannelIdx;
+			
+			float Energy = ChannelEnergyHistory[HistoryIndex];
+			const FAggregateLoudnessStats& Stats = AggregateLoudnessStats[ChannelIdx];
+
+			OutResult.Add(
+				FLKFSResults 
+				{
+					.Channel=ChannelIdx, 
+					.Timestamp=Timestamp, 
+					.Energy=Energy, 
+					.Loudness=FLoudnessAnalyzer::ConvertPerceptualEnergyToLKFS(Energy),
+					.ShortTermLoudness=Stats.ShortTermLoudness, 
+					.IntegratedLoudness=Stats.IntegratedLoudness,
+					.GatedLoudness=Stats.GatedLoudness
+				}
+			);
+		}
+		{
+			int32 HistoryIndex = (NumEnergyHistory - 1);
+			
+			float Energy = OverallEnergyHistory[HistoryIndex];
+			const FAggregateLoudnessStats& Stats = AggregateLoudnessStats[FLKFSAnalyzerResult::ChannelIndexOverall];
+
+			OutResult.Add(
+				FLKFSResults 
+				{
+					.Channel=FLKFSAnalyzerResult::ChannelIndexOverall, 
+					.Timestamp=Timestamp, 
+					.Energy=Energy, 
+					.Loudness=FLoudnessAnalyzer::ConvertPerceptualEnergyToLKFS(Energy),
+					.ShortTermLoudness=Stats.ShortTermLoudness, 
+					.IntegratedLoudness=Stats.IntegratedLoudness,
+					.GatedLoudness=Stats.GatedLoudness
+				}
+			);
+		}
+	}
+
+	void FLKFSWorker::UpdateAggregateLoudnessStats()
+	{
+		// Short term loudness is always updated, but integrated & gated loudness are 
+		// updated at a lower rate. 
+		bool bUpdateIntegratedLoudness = (NumEnergyHistory == 1) || (NumEnergyHistory % NumIntegratedHopHistoryResults) == 0;
+
+
+		// Per channel results
+		for (int32 ChannelIdx = 0; ChannelIdx < NumChannels; ChannelIdx++)
+		{
+			FAggregateLoudnessStats& Stats = AggregateLoudnessStats[ChannelIdx];
+			UpdateAggregateLoudnessStats(ChannelEnergyHistory, ChannelIdx, NumChannels, bUpdateIntegratedLoudness, Stats);
+		}
+
+		// Overall results
+		{
+			FAggregateLoudnessStats& Stats = AggregateLoudnessStats[FLKFSAnalyzerResult::ChannelIndexOverall];
+			UpdateAggregateLoudnessStats(OverallEnergyHistory, 0 /*offset*/, 1/*stride*/, bUpdateIntegratedLoudness, Stats);
+		}
+	}
+
+	void FLKFSWorker::UpdateAggregateLoudnessStats(const TCircularBuffer<float>& InHistory, int32 InOffset, int32 InStride, bool bInUpdateIntegratedLoudness, FAggregateLoudnessStats& OutStats) const
+	{
+		const int32 EndIndex = (NumEnergyHistory * InStride) + InOffset;
+
+		// Update short term loudness always.
+		{
+			const int32 StartIndex = InOffset + FMath::Max(0, (NumEnergyHistory - NumShortTermHistoryResults) * InStride);
+			float ShortTermAverageEnergy = LKFSFactoryPrivate::CalculateAverageFromHistory(InHistory, StartIndex, EndIndex, InStride);
+			OutStats.ShortTermLoudness = FLoudnessAnalyzer::ConvertPerceptualEnergyToLKFS(ShortTermAverageEnergy);
+		}
+
+		// Only update long term loudness if flag is set. 
+		if (bInUpdateIntegratedLoudness)
+		{
+			const int32 StartIndex = InOffset + FMath::Max(0, (NumEnergyHistory - NumIntegratedHistoryResults * InStride));
+
+			float IntegratedEnergy = LKFSFactoryPrivate::CalculateAverageFromHistory(InHistory, StartIndex, EndIndex, InStride);
+			OutStats.IntegratedLoudness = FLoudnessAnalyzer::ConvertPerceptualEnergyToLKFS(IntegratedEnergy);
+			OutStats.GatedLoudness = LKFSFactoryPrivate::CalculateGatedLoudnessFromHistory(InHistory, StartIndex, EndIndex, InStride);
+		}
+	}
+
+	/** Name of specific analyzer type. */
+	FName FLKFSFactory::GetName() const 
+	{
+		static FName FactoryName(TEXT("LKFSFactory"));
+		return FactoryName;
+	}
+
+	/** Human readable name of analyzer. */
+	FString FLKFSFactory::GetTitle() const
+	{
+		return TEXT("LKFS Analyzer Real-Time");
+	}
+
+	TUniquePtr<IAnalyzerResult> FLKFSFactory::NewResult() const
+	{
+		TUniquePtr<FLKFSAnalyzerResult> Result = MakeUnique<FLKFSAnalyzerResult>();
+		return Result;
+	}
+
+	TUniquePtr<IAnalyzerWorker> FLKFSFactory::NewWorker(const FAnalyzerParameters& InParams, const IAnalyzerSettings* InSettings) const
+	{
+		const FLKFSSettings* LoudnessSettings = static_cast<const FLKFSSettings*>(InSettings);
+
+		check(nullptr != LoudnessSettings);
+
+		return MakeUnique<FLKFSWorker>(InParams, *LoudnessSettings);
+	}
+}
+
diff --git a/Engine/Plugins/Runtime/AudioSynesthesia/Source/AudioSynesthesiaCore/Private/LKFSNRTFactory.cpp b/Engine/Plugins/Runtime/AudioSynesthesia/Source/AudioSynesthesiaCore/Private/LKFSNRTFactory.cpp
--- a/Engine/Plugins/Runtime/AudioSynesthesia/Source/AudioSynesthesiaCore/Private/LKFSNRTFactory.cpp
+++ b/Engine/Plugins/Runtime/AudioSynesthesia/Source/AudioSynesthesiaCore/Private/LKFSNRTFactory.cpp
@@ -5,6 +5,7 @@
 #include "Algo/IsSorted.h"
 #include "AudioSynesthesiaCoreLog.h"
 #include "DSP/SlidingWindow.h"
+#include "LKFSPrivate.h"
 
 #include UE_INLINE_GENERATED_CPP_BY_NAME(LKFSNRTFactory)
 
@@ -256,7 +257,7 @@ namespace Audio
 		// FLoundessDatum. The sliding window hop is tuned here so that it best 
 		// matches the desired AnalysisPeriod while maintaining between a 25% to 
 		// 75% window overlap.
-		TuneSlidingWindwoHopSize(NumAnalysisHopFrames, NumSlidingWindowFrames, NumSlidingWindowHopFrames, NumAnalysisHopWindows);
+		LKFSPrivate::TuneSlidingWindwoHopSize(NumAnalysisHopFrames, NumSlidingWindowFrames, NumSlidingWindowHopFrames, NumAnalysisHopWindows);
 
 		// Determine how many windows to analyze per loudness result
 		NumAnalysisWindows = FMath::Max(1, ((NumAnalysisFrames - (NumSlidingWindowFrames / 2)) / NumSlidingWindowHopFrames) + 1);
@@ -334,51 +335,6 @@ namespace Audio
 		}
 	}
 
-	void FLKFSNRTWorker::TuneSlidingWindwoHopSize(const int32 InAnalysisHopFrames, const int32 InWindowFrames, int32& OutBestWindowHopFrames, int32& OutBestAnalysisHopWindows) const
-	{
-		// The minimum sliding window hop size is the lesser of 
-		// - 25% of the sliding window size
-		// - The analysis hop size
-		int32 MinNumSlidingWindowHopFrames = FMath::Max(1, InWindowFrames / 4);
-		if (MinNumSlidingWindowHopFrames > InAnalysisHopFrames)
-		{
-			MinNumSlidingWindowHopFrames = InAnalysisHopFrames;
-		}
-
-		// The maximum sliding window hop size is the lesser of
-		// - 75% of the sliding window size
-		// - The analysis hop size
-		int32 MaxNumSlidingWindowHopFrames = FMath::Max(1, (3 * InWindowFrames) / 4);
-		if (MaxNumSlidingWindowHopFrames > InAnalysisHopFrames)
-		{
-			MaxNumSlidingWindowHopFrames = InAnalysisHopFrames;
-		}
-
-		// Iterate through the possible sizes to find the one that fits best with our settings.
-		int32 BestWindowHopFrames = -1;
-		int32 BestAnalysisHopWindows = -1;
-		int32 BestMismatch = InAnalysisHopFrames + 2; 
-
-		for (int32 CandidateWindowHopFrames = MinNumSlidingWindowHopFrames; CandidateWindowHopFrames <= MaxNumSlidingWindowHopFrames; CandidateWindowHopFrames++)
-		{
-			int32 CandidateAnalysisHopWindows = InAnalysisHopFrames / CandidateWindowHopFrames;
-			// The match quality is determined by how neatly the sliding window hop size fits into the anlaysis window hop size. 
-			int32 CandidateMismatch = FMath::Abs<int32>(InAnalysisHopFrames - (CandidateWindowHopFrames * CandidateAnalysisHopWindows));
-			if (CandidateMismatch < BestMismatch)
-			{
-				BestMismatch = CandidateMismatch;
-				BestWindowHopFrames = CandidateWindowHopFrames;
-				BestAnalysisHopWindows = CandidateAnalysisHopWindows;
-			}
-		}
-
-		// The algorithm above should always find a valid window hop
-		check(BestWindowHopFrames != -1);
-
-		OutBestWindowHopFrames = BestWindowHopFrames;
-		OutBestAnalysisHopWindows = BestAnalysisHopWindows;
-	}
-
 	void FLKFSNRTWorker::AnalyzeSlidingWindow(TAutoSlidingWindow<float>& InSlidingWindow, FLKFSNRTResult& OutResult)
 	{
 		check(NumAnalysisWindows > 0);
diff --git a/Engine/Plugins/Runtime/AudioSynesthesia/Source/AudioSynesthesiaCore/Private/LKFSPrivate.cpp b/Engine/Plugins/Runtime/AudioSynesthesia/Source/AudioSynesthesiaCore/Private/LKFSPrivate.cpp
new file mode 100644
--- /dev/null
+++ b/Engine/Plugins/Runtime/AudioSynesthesia/Source/AudioSynesthesiaCore/Private/LKFSPrivate.cpp
@@ -0,0 +1,54 @@
+// Copyright Epic Games, Inc. All Rights Reserved.
+
+#include "LKFSPrivate.h"
+
+#include "Math/UnrealMathUtility.h"
+
+namespace Audio::LKFSPrivate
+{
+	void TuneSlidingWindwoHopSize(int32 InAnalysisHopFrames, int32 InWindowFrames, int32& OutBestWindowHopFrames, int32& OutBestAnalysisHopWindows) 
+	{
+		// The minimum sliding window hop size is the lesser of 
+		// - 25% of the sliding window size
+		// - The analysis hop size
+		int32 MinNumSlidingWindowHopFrames = FMath::Max(1, InWindowFrames / 4);
+		if (MinNumSlidingWindowHopFrames > InAnalysisHopFrames)
+		{
+			MinNumSlidingWindowHopFrames = InAnalysisHopFrames;
+		}
+
+		// The maximum sliding window hop size is the lesser of
+		// - 75% of the sliding window size
+		// - The analysis hop size
+		int32 MaxNumSlidingWindowHopFrames = FMath::Max(1, (3 * InWindowFrames) / 4);
+		if (MaxNumSlidingWindowHopFrames > InAnalysisHopFrames)
+		{
+			MaxNumSlidingWindowHopFrames = InAnalysisHopFrames;
+		}
+
+		// Iterate through the possible sizes to find the one that fits best with our settings.
+		int32 BestWindowHopFrames = -1;
+		int32 BestAnalysisHopWindows = -1;
+		int32 BestMismatch = InAnalysisHopFrames + 2; 
+
+		for (int32 CandidateWindowHopFrames = MinNumSlidingWindowHopFrames; CandidateWindowHopFrames <= MaxNumSlidingWindowHopFrames; CandidateWindowHopFrames++)
+		{
+			int32 CandidateAnalysisHopWindows = InAnalysisHopFrames / CandidateWindowHopFrames;
+			// The match quality is determined by how neatly the sliding window hop size fits into the anlaysis window hop size. 
+			int32 CandidateMismatch = FMath::Abs<int32>(InAnalysisHopFrames - (CandidateWindowHopFrames * CandidateAnalysisHopWindows));
+			if (CandidateMismatch < BestMismatch)
+			{
+				BestMismatch = CandidateMismatch;
+				BestWindowHopFrames = CandidateWindowHopFrames;
+				BestAnalysisHopWindows = CandidateAnalysisHopWindows;
+			}
+		}
+
+		// The algorithm above should always find a valid window hop
+		check(BestWindowHopFrames != -1);
+
+		OutBestWindowHopFrames = BestWindowHopFrames;
+		OutBestAnalysisHopWindows = BestAnalysisHopWindows;
+	}
+}
+
diff --git a/Engine/Plugins/Runtime/AudioSynesthesia/Source/AudioSynesthesiaCore/Private/LKFSPrivate.h b/Engine/Plugins/Runtime/AudioSynesthesia/Source/AudioSynesthesiaCore/Private/LKFSPrivate.h
new file mode 100644
--- /dev/null
+++ b/Engine/Plugins/Runtime/AudioSynesthesia/Source/AudioSynesthesiaCore/Private/LKFSPrivate.h
@@ -0,0 +1,13 @@
+// Copyright Epic Games, Inc. All Rights Reserved.
+
+#pragma once
+#include "HAL/Platform.h"
+
+namespace Audio::LKFSPrivate
+{
+	// There may be multiple calls to the LoudnessAnalyzer to produce a single 
+	// FLoundessDatum. The sliding window hop is tuned here so that it best 
+	// matches the desired AnalysisPeriod while maintaining between a 25% to 
+	// 75% window overlap.
+	void TuneSlidingWindwoHopSize(int32 InNumAnalysisHopFrames, int32 InNumSlidingWindowFrames, int32& InNumSlidingWindowHopFrames, int32& InNumAnalysisHopWindows);
+}
diff --git a/Engine/Plugins/Runtime/AudioSynesthesia/Source/AudioSynesthesiaCore/Public/LKFSFactory.h b/Engine/Plugins/Runtime/AudioSynesthesia/Source/AudioSynesthesiaCore/Public/LKFSFactory.h
new file mode 100644
--- /dev/null
+++ b/Engine/Plugins/Runtime/AudioSynesthesia/Source/AudioSynesthesiaCore/Public/LKFSFactory.h
@@ -0,0 +1,209 @@
+// Copyright Epic Games, Inc. All Rights Reserved.
+
+#pragma once
+
+#include "Containers/ArrayView.h"
+#include "Containers/CircularBuffer.h"
+#include "Containers/SortedMap.h"
+#include "DSP/SlidingWindow.h"
+#include "LoudnessAnalyzer.h"
+#include "IAudioAnalyzerInterface.h"
+
+#include "LKFSFactory.generated.h"
+
+namespace Audio { template <typename InSampleType> class TSlidingBuffer; }
+
+/** Data representing the instantaneous loudness of a audio. */
+USTRUCT(BlueprintType)
+struct FLKFSResults
+{
+	GENERATED_BODY()
+
+	/** The audio channel index which produced this result. If -1, then
+	 * the data was produced by combining all other audio channels */
+	UPROPERTY(EditAnywhere, BlueprintReadOnly, Category = AudioAnalyzer)
+	int32 Channel = 0;
+
+	/** Time in seconds of the source audio which corresponds to the loudness measurements. */
+	UPROPERTY(EditAnywhere, BlueprintReadOnly, Category = AudioAnalyzer)
+	float Timestamp = 0.f;
+
+	/** The instantaneous, perceptually weighted energy relative to full scale. */
+	UPROPERTY(EditAnywhere, BlueprintReadOnly, Category = AudioAnalyzer)
+	float Energy = 0.f;
+
+	/** The instantaneous, perceptually weighted loudness in dB. */
+	UPROPERTY(EditAnywhere, BlueprintReadOnly, Category = AudioAnalyzer)
+	float Loudness = 0.f;
+
+	/** The average loudness over a short term window of time in dB. */
+	UPROPERTY(EditAnywhere, BlueprintReadOnly, Category = AudioAnalyzer)
+	float ShortTermLoudness = 0.f;
+
+	/** The average loudness over a long term window of time in dB. */
+	UPROPERTY(EditAnywhere, BlueprintReadOnly, Category = AudioAnalyzer)
+	float IntegratedLoudness = 0.f;
+
+	/** The average gated loudness over a long term window of time in dB. */
+	UPROPERTY(EditAnywhere, BlueprintReadOnly, Category = AudioAnalyzer)
+	float GatedLoudness = 0.f;
+};
+
+
+namespace Audio
+{
+	/**
+	 * Contains settings for loudness analyzer.
+	 */
+	class FLKFSSettings : public IAnalyzerSettings 
+	{
+	public:
+		// These values are taken from the suggested standards for LKFS/LUFS defined
+		// in ITU-R BS.1770
+		static constexpr float StandardShortTermLoudnessDuration = 3.f; // 3 seconds
+		static constexpr float StandardIntegratedLoudnessAnalysisPeriod = 1.f; // 1 second
+																			   //
+		// Streaming audio does not have a natural end point. Integrated loudness
+		// in this implementation is taken over a long duration (as opposed to 
+		// accumulating indefinitely). 
+		static constexpr float DefaultIntegratedLoudnessDuration = 60.f; // 1 minute
+
+		FLoudnessAnalyzerSettings LoudnessAnalyzerSettings = { 
+			.FFTSize = 2048,
+			.WindowType = EWindowType::Blackman,
+			.LoudnessCurveType = ELoudnessCurveType::K,
+			.MinAnalysisFrequency = 20.f,
+			.MaxAnalysisFrequency = 20000.f,
+			.ScalingMethod = ELoudnessAnalyzerScalingMethod::Corrected
+		};
+
+		/** Number of seconds between loudness measurements */
+		float AnalysisPeriod = 0.1f;
+
+		/** Number of seconds of audio analyzed for each loudness measurements */
+		float AnalysisWindowDuration = 0.4f;
+
+		/** Duration of audio analyzed for short term loudness. */
+		float ShortTermLoudnessDuration = StandardShortTermLoudnessDuration;
+
+		/** Number of seconds between long term loudness updates. */
+		float IntegratedLoudnessAnalysisPeriod = StandardIntegratedLoudnessAnalysisPeriod;
+
+		/** Duration of audio analyzed for long term loudness. */
+		float IntegratedLoudnessDuration = DefaultIntegratedLoudnessDuration;
+	};
+
+
+	/** FLKFSResult contains the temporal evolution of loudness. */
+	class FLKFSAnalyzerResult : public IAnalyzerResult
+	{
+	public:
+		/** Denotes the overall loudness channel index as opposed individual channel indices. */
+		AUDIOSYNESTHESIACORE_API static const int32 ChannelIndexOverall;
+
+		AUDIOSYNESTHESIACORE_API FLKFSAnalyzerResult();
+
+		/** The number of individual channels available in the result. */
+		AUDIOSYNESTHESIACORE_API int32 GetNumChannels() const;
+
+		/** Add to the result. */
+		AUDIOSYNESTHESIACORE_API void Add(FLKFSResults InResults);
+
+		/** Returns const reference to FLKFSResults array for individual channel. */
+		AUDIOSYNESTHESIACORE_API const TArray<FLKFSResults>& GetChannelLoudnessResults(int32 ChannelIdx) const;
+
+		/** Returns const reference to FLKFSResults array associated with overall loudness. */
+		AUDIOSYNESTHESIACORE_API const TArray<FLKFSResults>& GetLoudnessResults() const;
+
+	private:
+		TSortedMap<int32, TArray<FLKFSResults>> ChannelResults;
+	};
+
+	/** 
+	 * FLKFSWorker performs loudness analysis on input sample buffers.
+	 */
+	class FLKFSWorker : public IAnalyzerWorker
+	{
+	public:
+		/** Construct a worker */
+		AUDIOSYNESTHESIACORE_API FLKFSWorker(const FAnalyzerParameters& InParams, const FLKFSSettings& InAnalyzerSettings);
+
+		/**
+		 * Analyzes input sample buffer and updates result. 
+		 */
+		AUDIOSYNESTHESIACORE_API virtual void Analyze(TArrayView<const float> InAudio, IAnalyzerResult* OutResult) override;
+
+	private:
+		struct FAggregateLoudnessStats
+		{
+			static constexpr float InvalidLoudness = TNumericLimits<float>::Lowest();
+
+			float ShortTermLoudness = InvalidLoudness;
+			float IntegratedLoudness = InvalidLoudness;
+			float GatedLoudness = InvalidLoudness;
+		};
+
+
+		void TuneSlidingWindwoHopSize(const int32 InAnalysisHopFrames, const int32 InWindowFrames, int32& OutBestWindowHopFrames, int32& OutBestAnalysisHopWindows) const;
+		void AnalyzeSlidingWindow(TAutoSlidingWindow<float>& InSlidingWindow, FLKFSAnalyzerResult& OutResult);
+		void AnalyzeWindow(TArrayView<const float> InWindow);
+
+		void AddCurrentLoudnessDataToHistory();
+		void AddCurrentLoudnessDataToResult(FLKFSAnalyzerResult& OutResult);
+		void UpdateAggregateLoudnessStats();
+		void UpdateAggregateLoudnessStats(const TCircularBuffer<float>& InHistory, int32 InOffset, int32 InStride, bool bInUpdateIntegratedLoudness, FAggregateLoudnessStats& OutStats) const;
+
+
+		int32 NumChannels = 0;
+		int32 NumAnalyzedBuffers = 0;
+		int32 NumEnergyHistory = 0;
+		int32 NumSlidingWindowHopFrames = 0;
+		float SampleRate = 0.f;
+
+		// Number of windows between each loudness analysis
+		int32 NumAnalysisHopWindows = 1;
+		// Number of windows inspected for each loudness analysis
+		int32 NumAnalysisWindows = 1;
+		// Number of values in the "EnergyHistory" to analyze for short term loudness
+		int32 NumShortTermHistoryResults = 1;
+		// Number of values in the "EnergyHistory" to analyze for integrated and gated loudness
+		int32 NumIntegratedHistoryResults= 1;
+		// Number of values in the "EnergyHistory" between integrated and gated loudness anlysis
+		int32 NumIntegratedHopHistoryResults = 1;
+
+		TArray<float> ChannelPerceptualEnergy;
+		TArray<float> OverallPerceptualEnergy;
+		TArray<float> InternalWindow;
+
+		// Per channel stats for short-term, integrated and gated loudness. 
+		TSortedMap<int32, FAggregateLoudnessStats> AggregateLoudnessStats;
+
+		// Channel interleaved history buffer of energy measurements. These are 
+		// taken once per an loudness analysis, not necessarily once per a audio window. 
+		TCircularBuffer<float> ChannelEnergyHistory;
+		TCircularBuffer<float> OverallEnergyHistory;
+
+		TSlidingBuffer<float> InternalBuffer;
+		FMultichannelLoudnessAnalyzer Analyzer;
+	};
+
+	/**
+	 * Defines the LKFS analyzer and creates related classes.
+	 */
+	class FLKFSFactory : public IAnalyzerFactory
+	{
+	public:
+
+		/** Name of specific analyzer type. */
+		AUDIOSYNESTHESIACORE_API virtual FName GetName() const override;
+
+		/** Human readable name of analyzer. */
+		AUDIOSYNESTHESIACORE_API virtual FString GetTitle() const override;
+
+		/** Creates a new FLKFSAnalyzerResult */
+		AUDIOSYNESTHESIACORE_API virtual TUniquePtr<IAnalyzerResult> NewResult() const override;
+
+		/** Creates a new FLKFSWorker. This expects IAnalyzerSettings to be a valid pointer to a FLKFSSettings object. */
+		AUDIOSYNESTHESIACORE_API virtual TUniquePtr<IAnalyzerWorker> NewWorker(const FAnalyzerParameters& InParams, const IAnalyzerSettings* InSettings) const override;
+	};
+}
diff --git a/Engine/Plugins/Runtime/AudioSynesthesia/Source/AudioSynesthesiaCore/Public/LKFSNRTFactory.h b/Engine/Plugins/Runtime/AudioSynesthesia/Source/AudioSynesthesiaCore/Public/LKFSNRTFactory.h
--- a/Engine/Plugins/Runtime/AudioSynesthesia/Source/AudioSynesthesiaCore/Public/LKFSNRTFactory.h
+++ b/Engine/Plugins/Runtime/AudioSynesthesia/Source/AudioSynesthesiaCore/Public/LKFSNRTFactory.h
@@ -197,7 +197,6 @@ namespace Audio
 	private:
 
 
-		void TuneSlidingWindwoHopSize(const int32 InAnalysisHopFrames, const int32 InWindowFrames, int32& OutBestWindowHopFrames, int32& OutBestAnalysisHopWindows) const;
 		void AnalyzeSlidingWindow(TAutoSlidingWindow<float>& InSlidingWindow, FLKFSNRTResult& OutResult);
 		void AnalyzeWindow(TArrayView<const float> InWindow, FLKFSNRTResult& OutResult);
 
diff --git a/Engine/Source/Runtime/AudioAnalyzer/Private/AudioAnalyzerSubsystem.cpp b/Engine/Source/Runtime/AudioAnalyzer/Private/AudioAnalyzerSubsystem.cpp
--- a/Engine/Source/Runtime/AudioAnalyzer/Private/AudioAnalyzerSubsystem.cpp
+++ b/Engine/Source/Runtime/AudioAnalyzer/Private/AudioAnalyzerSubsystem.cpp
@@ -38,11 +38,14 @@ bool UAudioAnalyzerSubsystem::Tick(float DeltaTime)
 	// Loop through all analyzers and if they're ready to analyze, do it
 	for (UAudioAnalyzer* Analyzer : AudioAnalyzers)
 	{
-		if (Analyzer->IsReadyForAnalysis())
+		if (Analyzer)
 		{
-			if (Analyzer->DoAnalysis())
+			if (Analyzer->IsReadyForAnalysis())
 			{
-				Analyzer->BroadcastResults();
+				if (Analyzer->DoAnalysis())
+				{
+					Analyzer->BroadcastResults();
+				}
 			}
 		}
 	}
diff --git a/Engine/Source/Runtime/SignalProcessing/Public/DSP/SlidingWindow.h b/Engine/Source/Runtime/SignalProcessing/Public/DSP/SlidingWindow.h
--- a/Engine/Source/Runtime/SignalProcessing/Public/DSP/SlidingWindow.h
+++ b/Engine/Source/Runtime/SignalProcessing/Public/DSP/SlidingWindow.h
@@ -26,10 +26,6 @@ namespace Audio
 	template <typename InSampleType>
 	class TSlidingBuffer
 	{
-			// Give TSlidingWindow access to StorageBuffer without exposing StorageBuffer public interface.
-			friend class TSlidingWindow<InSampleType>;
-
-
 		public:
 
 			/**
@@ -44,6 +40,17 @@ namespace Audio
 				check(NumHopSamples >= 1);
 			}
 
+			TSlidingBuffer()
+			: TSlidingBuffer(1, 1)
+			{
+			}
+
+			TSlidingBuffer(const TSlidingBuffer&) = default;
+			TSlidingBuffer(TSlidingBuffer&&) = default;
+
+			TSlidingBuffer& operator=(const TSlidingBuffer& InOther) = default;
+			TSlidingBuffer& operator=(TSlidingBuffer&& InOther) = default;
+
 			/** Returns the number of samples in a window. */
 			int32 GetNumWindowSamples() const
 			{
@@ -56,6 +63,17 @@ namespace Audio
 				return NumHopSamples;
 			}
 
+			/** Returns the number of underflow samples (when hopsize > windowsize) */
+			int32 GetNumUnderflowSamples() const
+			{
+				return NumUnderflowSamples;
+			}
+
+			TArrayView<const InSampleType> GetStorage() const
+			{
+				return StorageBuffer;
+			}
+
 			/**
 			 * StoreForFutureWindows stores the necessary samples from InBuffer which will
 			 * be needed for future windows. It ignores all values in InBuffer which can
@@ -313,12 +331,12 @@ namespace Audio
 			 * bDoFlush Controls whether zeros to the final output windows until all possible windows with data from InNewBuffer have been covered.
 			 */
 			TSlidingWindow(const TSlidingBuffer<InSampleType>& InSlidingBuffer, TArrayView<const InSampleType> InNewBuffer, bool bDoFlush)
-			:	StorageBuffer(InSlidingBuffer.StorageBuffer)
+			:	StorageBuffer(InSlidingBuffer.GetStorage())
 			,	NewBuffer(InNewBuffer)
-			,	NumWindowSamples(InSlidingBuffer.NumWindowSamples)
-			,	NumHopSamples(InSlidingBuffer.NumHopSamples)
+			,	NumWindowSamples(InSlidingBuffer.GetNumWindowSamples())
+			,	NumHopSamples(InSlidingBuffer.GetNumHopSamples())
 			,	MaxReadIndex(0)
-			,	NumUnderflowSamples(InSlidingBuffer.NumUnderflowSamples)
+			,	NumUnderflowSamples(InSlidingBuffer.GetNumUnderflowSamples())
 			{
 				// Total samples to be slid over.
 				int32 NumSamples = NewBuffer.Num() + StorageBuffer.Num();
