[Audio Widgets] Loudness Meter - Added Loudness Range (LRA) as defined in EBU Tech 3342 (editor only). Submit on behalf of pete.bliss + added LRA test case

#Synced-CL 50100174

diff --git a/Engine/Plugins/Runtime/AudioSynesthesia/Source/AudioSynesthesia/Classes/LKFS.h b/Engine/Plugins/Runtime/AudioSynesthesia/Source/AudioSynesthesia/Classes/LKFS.h
--- a/Engine/Plugins/Runtime/AudioSynesthesia/Source/AudioSynesthesia/Classes/LKFS.h
+++ b/Engine/Plugins/Runtime/AudioSynesthesia/Source/AudioSynesthesia/Classes/LKFS.h
@@ -43,6 +43,14 @@ public:
 	UPROPERTY(EditAnywhere, BlueprintReadOnly, Category=AudioAnalyzer, meta = (ClampMin = "10.00", ClampMax = "3600.00"))
 	float IntegratedLoudnessDuration = Audio::FLKFSSettings::DefaultIntegratedLoudnessDuration;
 
+	/** Should the Loudness Range (LRA) statistic be calculated for each channel? */
+	UPROPERTY(EditAnywhere, BlueprintReadOnly, Category = AudioAnalyzer)
+	bool bCalculatePerChannelLoudnessRange = false;
+
+	/** Should the overall Loudness Range (LRA) statistic be calculated? */
+	UPROPERTY(EditAnywhere, BlueprintReadOnly, Category = AudioAnalyzer)
+	bool bCalculateOverallLoudnessRange = false;
+
 	/** Convert ULKFSSettings to IAnalyzerSettings */
 	UE_API TUniquePtr<Audio::IAnalyzerSettings> GetSettings(const int32 InSampleRate, const int32 InNumChannels) const;
 
diff --git a/Engine/Plugins/Runtime/AudioSynesthesia/Source/AudioSynesthesia/Private/LKFS.cpp b/Engine/Plugins/Runtime/AudioSynesthesia/Source/AudioSynesthesia/Private/LKFS.cpp
--- a/Engine/Plugins/Runtime/AudioSynesthesia/Source/AudioSynesthesia/Private/LKFS.cpp
+++ b/Engine/Plugins/Runtime/AudioSynesthesia/Source/AudioSynesthesia/Private/LKFS.cpp
@@ -14,6 +14,8 @@ TUniquePtr<Audio::IAnalyzerSettings> ULKFSSettings::GetSettings(const int32 InSa
 	Settings->ShortTermLoudnessDuration = ShortTermLoudnessDuration;
 	Settings->IntegratedLoudnessAnalysisPeriod = IntegratedLoudnessAnalysisPeriod;
 	Settings->IntegratedLoudnessDuration = IntegratedLoudnessDuration;
+	Settings->bCalculatePerChannelLoudnessRange = bCalculatePerChannelLoudnessRange;
+	Settings->bCalculateOverallLoudnessRange = bCalculateOverallLoudnessRange;
 
 	return Settings;
 }
diff --git a/Engine/Plugins/Runtime/AudioSynesthesia/Source/AudioSynesthesiaCore/Private/LKFSFactory.cpp b/Engine/Plugins/Runtime/AudioSynesthesia/Source/AudioSynesthesiaCore/Private/LKFSFactory.cpp
--- a/Engine/Plugins/Runtime/AudioSynesthesia/Source/AudioSynesthesiaCore/Private/LKFSFactory.cpp
+++ b/Engine/Plugins/Runtime/AudioSynesthesia/Source/AudioSynesthesiaCore/Private/LKFSFactory.cpp
@@ -74,6 +74,76 @@ namespace Audio
 			const float GatedEnergyPass2 = CalculateGatedAverageFromHistory(InHistory, InStartIndex, InEndIndex, InStride, NewThreshold);
 			return FLoudnessAnalyzer::ConvertPerceptualEnergyToLKFS(GatedEnergyPass2);
 		}
+
+		/**
+		 * Implementation copied from FFineGrainedPerformanceTracker::GetPercentileValue(...)
+		 */
+		float GetPercentileValue(TArray<float>& Samples, int32 Percentile)
+		{
+			ensure(Percentile >= 0 && Percentile <= 100);
+
+			int32 Left = 0;
+			int32 Right = Samples.Num() - 1;
+
+			if (Right < 0)
+			{
+				return -1.0f;
+			}
+
+			const int32 PercentileClamped = FMath::Clamp(Percentile, 0, 100);
+			const int32 PercentileOrdinal = FMath::RoundToInt32((float(PercentileClamped) * Right) / 100);
+
+			// this is quickselect (see http://en.wikipedia.org/wiki/Quickselect for details).
+			while (Right != Left)
+			{
+				// partition
+				int32 MovingLeft = Left - 1;
+				int32 MovingRight = Right;
+				float Pivot = Samples[MovingRight];
+				for (; ; )
+				{
+					while (Samples[++MovingLeft] < Pivot);
+					while (Samples[--MovingRight] > Pivot)
+					{
+						if (MovingRight == Left)
+						{
+							break;
+						}
+					}
+
+					if (MovingLeft >= MovingRight)
+					{
+						break;
+					}
+
+					const float Temp = Samples[MovingLeft];
+					Samples[MovingLeft] = Samples[MovingRight];
+					Samples[MovingRight] = Temp;
+				}
+
+				const float Temp = Samples[MovingLeft];
+				Samples[MovingLeft] = Samples[Right];
+				Samples[Right] = Temp;
+
+				// now we're pivoted around MovingLeft
+				// decide what part K-th largest belong to
+				if (MovingLeft > PercentileOrdinal)
+				{
+					Right = MovingLeft - 1;
+				}
+				else if (MovingLeft < PercentileOrdinal)
+				{
+					Left = MovingLeft + 1;
+				}
+				else
+				{
+					// we hit exactly the value we need, no need to sort further
+					break;
+				}
+			}
+
+			return Samples[PercentileOrdinal];
+		}
 	}
 }
 
@@ -174,6 +244,18 @@ namespace Audio
 			AggregateLoudnessStats.Add(ChannelIdx);
 		}
 		AggregateLoudnessStats.Add(FLKFSAnalyzerResult::ChannelIndexOverall);
+
+		if (InAnalyzerSettings.bCalculatePerChannelLoudnessRange)
+		{
+			for (int32 ChannelIdx = 0; ChannelIdx < NumChannels; ChannelIdx++)
+			{
+				LoudnessRangeStatePerChannel.Emplace(ChannelIdx, static_cast<uint32>(NumIntegratedHistoryResults));
+			}
+		}
+		if (InAnalyzerSettings.bCalculateOverallLoudnessRange)
+		{
+			LoudnessRangeStatePerChannel.Emplace(FLKFSAnalyzerResult::ChannelIndexOverall, static_cast<uint32>(NumIntegratedHistoryResults));
+		}
 	}
 
 	void FLKFSWorker::Analyze(TArrayView<const float> InAudio, IAnalyzerResult* OutResult) 
@@ -277,7 +359,8 @@ namespace Audio
 					.Loudness=FLoudnessAnalyzer::ConvertPerceptualEnergyToLKFS(Energy),
 					.ShortTermLoudness=Stats.ShortTermLoudness, 
 					.IntegratedLoudness=Stats.IntegratedLoudness,
-					.GatedLoudness=Stats.GatedLoudness
+					.GatedLoudness=Stats.GatedLoudness,
+					.LoudnessRange=Stats.LoudnessRange,
 				}
 			);
 		}
@@ -296,7 +379,8 @@ namespace Audio
 					.Loudness=FLoudnessAnalyzer::ConvertPerceptualEnergyToLKFS(Energy),
 					.ShortTermLoudness=Stats.ShortTermLoudness, 
 					.IntegratedLoudness=Stats.IntegratedLoudness,
-					.GatedLoudness=Stats.GatedLoudness
+					.GatedLoudness=Stats.GatedLoudness,
+					.LoudnessRange=Stats.LoudnessRange,
 				}
 			);
 		}
@@ -313,36 +397,107 @@ namespace Audio
 		for (int32 ChannelIdx = 0; ChannelIdx < NumChannels; ChannelIdx++)
 		{
 			FAggregateLoudnessStats& Stats = AggregateLoudnessStats[ChannelIdx];
-			UpdateAggregateLoudnessStats(ChannelEnergyHistory, ChannelIdx, NumChannels, bUpdateIntegratedLoudness, Stats);
+			FLoudnessRange* LoudnessRange = LoudnessRangeStatePerChannel.Find(ChannelIdx);
+			UpdateAggregateLoudnessStats(ChannelEnergyHistory, ChannelIdx, NumChannels, bUpdateIntegratedLoudness, Stats, LoudnessRange);
 		}
 
 		// Overall results
 		{
 			FAggregateLoudnessStats& Stats = AggregateLoudnessStats[FLKFSAnalyzerResult::ChannelIndexOverall];
-			UpdateAggregateLoudnessStats(OverallEnergyHistory, 0 /*offset*/, 1/*stride*/, bUpdateIntegratedLoudness, Stats);
+			FLoudnessRange* LoudnessRange = LoudnessRangeStatePerChannel.Find(FLKFSAnalyzerResult::ChannelIndexOverall);
+			UpdateAggregateLoudnessStats(OverallEnergyHistory, 0 /*offset*/, 1/*stride*/, bUpdateIntegratedLoudness, Stats, LoudnessRange);
 		}
 	}
 
-	void FLKFSWorker::UpdateAggregateLoudnessStats(const TCircularBuffer<float>& InHistory, int32 InOffset, int32 InStride, bool bInUpdateIntegratedLoudness, FAggregateLoudnessStats& OutStats) const
+	void FLKFSWorker::UpdateAggregateLoudnessStats(const TCircularBuffer<float>& InHistory, int32 InOffset, int32 InStride, bool bInUpdateIntegratedLoudness, FAggregateLoudnessStats& OutStats, FLoudnessRange* LoudnessRange) const
 	{
 		const int32 EndIndex = (NumEnergyHistory * InStride) + InOffset;
 
 		// Update short term loudness always.
-		{
-			const int32 StartIndex = InOffset + FMath::Max(0, (NumEnergyHistory - NumShortTermHistoryResults) * InStride);
-			float ShortTermAverageEnergy = LKFSFactoryPrivate::CalculateAverageFromHistory(InHistory, StartIndex, EndIndex, InStride);
-			OutStats.ShortTermLoudness = FLoudnessAnalyzer::ConvertPerceptualEnergyToLKFS(ShortTermAverageEnergy);
-		}
+		const int32 ShortTermLoudnessStartIndex = InOffset + FMath::Max(0, (NumEnergyHistory - NumShortTermHistoryResults) * InStride);
+		const float ShortTermAverageEnergy = LKFSFactoryPrivate::CalculateAverageFromHistory(InHistory, ShortTermLoudnessStartIndex, EndIndex, InStride);
+		OutStats.ShortTermLoudness = FLoudnessAnalyzer::ConvertPerceptualEnergyToLKFS(ShortTermAverageEnergy);
 
 		// Only update long term loudness if flag is set. 
 		if (bInUpdateIntegratedLoudness)
 		{
-			const int32 StartIndex = InOffset + FMath::Max(0, (NumEnergyHistory - NumIntegratedHistoryResults * InStride));
+			const int32 LongTermLoudnessStartIndex = InOffset + FMath::Max(0, (NumEnergyHistory - NumIntegratedHistoryResults * InStride));
 
-			float IntegratedEnergy = LKFSFactoryPrivate::CalculateAverageFromHistory(InHistory, StartIndex, EndIndex, InStride);
+			float IntegratedEnergy = LKFSFactoryPrivate::CalculateAverageFromHistory(InHistory, LongTermLoudnessStartIndex, EndIndex, InStride);
 			OutStats.IntegratedLoudness = FLoudnessAnalyzer::ConvertPerceptualEnergyToLKFS(IntegratedEnergy);
-			OutStats.GatedLoudness = LKFSFactoryPrivate::CalculateGatedLoudnessFromHistory(InHistory, StartIndex, EndIndex, InStride);
+			OutStats.GatedLoudness = LKFSFactoryPrivate::CalculateGatedLoudnessFromHistory(InHistory, LongTermLoudnessStartIndex, EndIndex, InStride);
 		}
+
+		// Update optional loudness range (if we have LoudnessRange state):
+		if (LoudnessRange)
+		{
+			LoudnessRange->AddShortTermEnergyMeasurement(ShortTermAverageEnergy);
+
+			if (bInUpdateIntegratedLoudness)
+			{
+				OutStats.LoudnessRange = LoudnessRange->CalculateLoudnessRange();
+			}
+		}
+	}
+
+	void FLKFSWorker::FLoudnessRange::AddShortTermEnergyMeasurement(float InShortTermEnergy)
+	{
+		ShortTermEnergyHistory[ShortTermWriteIndex++] = InShortTermEnergy;
+		ShortTermHistorySize = FMath::Min<uint32>(ShortTermHistorySize + 1, MaxShortTermHistorySize);
+	}
+
+	TOptional<FFloatInterval> FLKFSWorker::FLoudnessRange::CalculateLoudnessRange() const
+	{
+		// Based on EBU Tech 3342
+		if (ShortTermHistorySize == 0)
+		{
+			return NullOpt;
+		}
+
+		constexpr float AbsoluteThresholdEnergy = 1.17246636e-7f; // -70 LUFS
+		constexpr float RelativeThresholdEnergy = 0.01f;		  // -20 LU 
+
+		const uint32 LoudnessRangeStartIndex = ShortTermWriteIndex - ShortTermHistorySize;
+		const uint32 LoudnessRangeEndIndex   = ShortTermWriteIndex;
+
+		// Apply the absolute-threshold gating
+		const float RefEnergy = LKFSFactoryPrivate::CalculateGatedAverageFromHistory(
+			ShortTermEnergyHistory,
+			static_cast<int32>(LoudnessRangeStartIndex),
+			static_cast<int32>(LoudnessRangeEndIndex),
+			1 /*Stride*/,
+			AbsoluteThresholdEnergy
+		);
+
+		// Apply the relative-threshold gating
+		const float RelGateEnergy = FMath::Max(RefEnergy * RelativeThresholdEnergy, AbsoluteThresholdEnergy);
+
+		TArray<float> ShortTermLoudnessRelGated;
+		ShortTermLoudnessRelGated.Reserve(ShortTermHistorySize);
+
+		for (uint32 Index = LoudnessRangeStartIndex; Index < LoudnessRangeEndIndex; ++Index)
+		{
+			const float ShortTermEnergy = ShortTermEnergyHistory[Index];
+
+			if (ShortTermEnergy >= RelGateEnergy)
+			{
+				ShortTermLoudnessRelGated.Emplace(FLoudnessAnalyzer::ConvertPerceptualEnergyToLKFS(ShortTermEnergy));
+			}
+		}
+
+		if (ShortTermLoudnessRelGated.IsEmpty())
+		{
+			return NullOpt;
+		}
+
+		// Compute the high and low percentiles of the distribution of values in ShortTermLoudnessRelGated
+		static constexpr int32 LowerPercentile = 10;
+		static constexpr int32 UpperPercentile = 95;
+
+		const float LoudnessRangeMin = LKFSFactoryPrivate::GetPercentileValue(ShortTermLoudnessRelGated, LowerPercentile);
+		const float LoudnessRangeMax = LKFSFactoryPrivate::GetPercentileValue(ShortTermLoudnessRelGated, UpperPercentile);
+
+		return FFloatInterval(LoudnessRangeMin, LoudnessRangeMax);
 	}
 
 	/** Name of specific analyzer type. */
diff --git a/Engine/Plugins/Runtime/AudioSynesthesia/Source/AudioSynesthesiaCore/Public/LKFSFactory.h b/Engine/Plugins/Runtime/AudioSynesthesia/Source/AudioSynesthesiaCore/Public/LKFSFactory.h
--- a/Engine/Plugins/Runtime/AudioSynesthesia/Source/AudioSynesthesiaCore/Public/LKFSFactory.h
+++ b/Engine/Plugins/Runtime/AudioSynesthesia/Source/AudioSynesthesiaCore/Public/LKFSFactory.h
@@ -47,6 +47,10 @@ struct FLKFSResults
 	/** The average gated loudness over a long term window of time in dB. */
 	UPROPERTY(EditAnywhere, BlueprintReadOnly, Category = AudioAnalyzer)
 	float GatedLoudness = 0.f;
+
+	/** Loudness Range (LRA), as defined in EBU Tech 3342. */
+	UPROPERTY(EditAnywhere, BlueprintReadOnly, Category = AudioAnalyzer)
+	TOptional<FFloatInterval> LoudnessRange;
 };
 
 
@@ -91,6 +95,12 @@ namespace Audio
 
 		/** Duration of audio analyzed for long term loudness. */
 		float IntegratedLoudnessDuration = DefaultIntegratedLoudnessDuration;
+
+		/** Should the Loudness Range (LRA) statistic be calculated for each channel? */
+		bool bCalculatePerChannelLoudnessRange = false;
+
+		/** Should the overall Loudness Range (LRA) statistic be calculated? */
+		bool bCalculateOverallLoudnessRange = false;
 	};
 
 
@@ -141,8 +151,32 @@ namespace Audio
 			float ShortTermLoudness = InvalidLoudness;
 			float IntegratedLoudness = InvalidLoudness;
 			float GatedLoudness = InvalidLoudness;
+			TOptional<FFloatInterval> LoudnessRange;
 		};
 
+		/**
+		 * Holds state for Loudness Range calculation (for a single channel).
+		 */
+		class FLoudnessRange
+		{
+		public:
+			explicit FLoudnessRange(uint32 InMaxShortTermHistorySize)
+				: ShortTermEnergyHistory(InMaxShortTermHistorySize)
+				, MaxShortTermHistorySize(InMaxShortTermHistorySize)
+			{
+			}
+			
+			void AddShortTermEnergyMeasurement(float InShortTermEnergy);
+
+			TOptional<FFloatInterval> CalculateLoudnessRange() const;
+
+		private:
+			TCircularBuffer<float> ShortTermEnergyHistory;
+
+			uint32 MaxShortTermHistorySize = 0;
+			uint32 ShortTermHistorySize = 0;
+			uint32 ShortTermWriteIndex  = 0;
+		};
 
 		void TuneSlidingWindwoHopSize(const int32 InAnalysisHopFrames, const int32 InWindowFrames, int32& OutBestWindowHopFrames, int32& OutBestAnalysisHopWindows) const;
 		void AnalyzeSlidingWindow(TAutoSlidingWindow<float>& InSlidingWindow, FLKFSAnalyzerResult& OutResult);
@@ -151,7 +185,7 @@ namespace Audio
 		void AddCurrentLoudnessDataToHistory();
 		void AddCurrentLoudnessDataToResult(FLKFSAnalyzerResult& OutResult);
 		void UpdateAggregateLoudnessStats();
-		void UpdateAggregateLoudnessStats(const TCircularBuffer<float>& InHistory, int32 InOffset, int32 InStride, bool bInUpdateIntegratedLoudness, FAggregateLoudnessStats& OutStats) const;
+		void UpdateAggregateLoudnessStats(const TCircularBuffer<float>& InHistory, int32 InOffset, int32 InStride, bool bInUpdateIntegratedLoudness, FAggregateLoudnessStats& OutStats, FLoudnessRange* LoudnessRange = nullptr) const;
 
 
 		int32 NumChannels = 0;
@@ -175,8 +209,10 @@ namespace Audio
 		TArray<float> OverallPerceptualEnergy;
 		TArray<float> InternalWindow;
 
-		// Per channel stats for short-term, integrated and gated loudness. 
+		// Per channel stats for short-term, integrated and gated loudness, and optional loudness range. 
 		TSortedMap<int32, FAggregateLoudnessStats> AggregateLoudnessStats;
+		// Per channel state for calculating Loudness Range.
+		TSortedMap<int32, FLoudnessRange> LoudnessRangeStatePerChannel;
 
 		// Channel interleaved history buffer of energy measurements. These are 
 		// taken once per an loudness analysis, not necessarily once per a audio window. 
diff --git a/Engine/Plugins/Runtime/AudioWidgets/Source/AudioWidgets/Private/AudioLoudnessMeter.cpp b/Engine/Plugins/Runtime/AudioWidgets/Source/AudioWidgets/Private/AudioLoudnessMeter.cpp
--- a/Engine/Plugins/Runtime/AudioWidgets/Source/AudioWidgets/Private/AudioLoudnessMeter.cpp
+++ b/Engine/Plugins/Runtime/AudioWidgets/Source/AudioWidgets/Private/AudioLoudnessMeter.cpp
@@ -100,6 +100,7 @@ namespace AudioWidgets
 		Teardown();
 
 		LKFSSettings = TStrongObjectPtr(NewObject<ULKFSSettings>());
+		LKFSSettings->bCalculateOverallLoudnessRange = true;
 
 		AudioDeviceId = InAudioDeviceId;
 
@@ -192,6 +193,13 @@ namespace AudioWidgets
 				.Value = TAttribute<TOptional<float>>::CreateSPLambda(this, [this]() { return (LatestLoudnessResults.IsSet()) ? TOptional<float>(LatestLoudnessResults->Loudness) : NullOpt; }),
 				.bShowValue = false,
 				.bShowMeter = true,
+			},
+			{
+				.Name = GET_MEMBER_NAME_CHECKED(FLoudnessMeterRackUnitSettings, LoudnessRange),
+				.DisplayName = LOCTEXT("LoudnessRangeDisplayName", "Range"),
+				.Range = TAttribute<TOptional<FFloatInterval>>::CreateSPLambda(this, [this]() { return (LatestLoudnessResults.IsSet()) ? LatestLoudnessResults->LoudnessRange : NullOpt; }),
+				.bShowValue = true,
+				.bShowMeter = false,
 			}
 		};
 
diff --git a/Engine/Plugins/Runtime/AudioWidgets/Source/AudioWidgets/Private/LoudnessMeterWidgetView.cpp b/Engine/Plugins/Runtime/AudioWidgets/Source/AudioWidgets/Private/LoudnessMeterWidgetView.cpp
--- a/Engine/Plugins/Runtime/AudioWidgets/Source/AudioWidgets/Private/LoudnessMeterWidgetView.cpp
+++ b/Engine/Plugins/Runtime/AudioWidgets/Source/AudioWidgets/Private/LoudnessMeterWidgetView.cpp
@@ -32,31 +32,31 @@ namespace AudioWidgets
 		using FLoudnessMetricRef = TSharedRef<const FLoudnessMeterWidgetView::FLoudnessMetric>;
 
 		/**
-		 * Derived meter widget type adds bindable attributes for scale Range and Offset.
+		 * Derived meter widget type adds bindable attributes for scale DisplayRange and Offset.
 		 */
-		class SVariableRangeMeterWidget : public SAudioMeterWidget
+		class SVariableDisplayRangeMeterWidget : public SAudioMeterWidget
 		{
 		public:
-			SLATE_BEGIN_ARGS(SVariableRangeMeterWidget)
+			SLATE_BEGIN_ARGS(SVariableDisplayRangeMeterWidget)
 				: _MeterValueColor(FLinearColor::Green)
-				, _Range(60)
+				, _DisplayRange(60)
 				, _Offset(0)
 				{
 				}
 
 				SLATE_ATTRIBUTE(FSlateColor, MeterValueColor)
 				SLATE_ATTRIBUTE(TArray<FAudioMeterChannelInfo>, MeterChannelInfo)
-				SLATE_ATTRIBUTE(int, Range)
+				SLATE_ATTRIBUTE(int, DisplayRange)
 				SLATE_ATTRIBUTE(int, Offset)
 
 			SLATE_END_ARGS()
 
-			void Construct(const SVariableRangeMeterWidget::FArguments& InArgs)
+			void Construct(const SVariableDisplayRangeMeterWidget::FArguments& InArgs)
 			{
 				AudioMeterWidgetStyle->SetMeterSize({ 700.0, 25.0 }); // MaxLength, Width.
 				AudioMeterWidgetStyle->SetScaleHashHeight(4.0f); // This is the tick mark length.
 
-				Range = InArgs._Range;
+				DisplayRange = InArgs._DisplayRange;
 				Offset = InArgs._Offset;
 
 				SAudioMeterWidget::Construct(
@@ -77,9 +77,9 @@ namespace AudioWidgets
 			virtual int32 OnPaint(const FPaintArgs& Args, const FGeometry& AllottedGeometry, const FSlateRect& MyCullingRect, FSlateWindowElementList& OutDrawElements, int32 LayerId, const FWidgetStyle& InWidgetStyle, bool bParentEnabled) const override
 			{
 				// Update the style from our attribute bindings:
-				const float RangeValue = Range.Get();
+				const float DisplayRangeValue = DisplayRange.Get();
 				const float OffsetValue = Offset.Get();
-				AudioMeterWidgetStyle->SetValueRangeDb({ OffsetValue - RangeValue, OffsetValue });
+				AudioMeterWidgetStyle->SetValueRangeDb({ OffsetValue - DisplayRangeValue, OffsetValue });
 
 				// Paint the widget using the updated style:
 				return SAudioMeterWidget::OnPaint(Args, AllottedGeometry, MyCullingRect, OutDrawElements, LayerId, InWidgetStyle, bParentEnabled);
@@ -88,7 +88,7 @@ namespace AudioWidgets
 		private:
 			const TUniquePtr<FAudioMeterWidgetStyle> AudioMeterWidgetStyle = MakeUnique<FAudioMeterWidgetStyle>(FAudioMeterWidgetStyle::GetDefault());
 			TAttribute<int> Offset;
-			TAttribute<int> Range;
+			TAttribute<int> DisplayRange;
 		};
 
 		/**
@@ -129,15 +129,31 @@ namespace AudioWidgets
 					return LOCTEXT("NegativeInfinity", "-∞");
 				}
 
-				static const FNumberFormattingOptions NumberFormattingOptions = FNumberFormattingOptions()
-					.SetMinimumFractionalDigits(1)
-					.SetMaximumFractionalDigits(1);
-
 				const float ConvertedValue = ConvertValue(*Value);
-				return FText::AsNumber(ConvertedValue, &NumberFormattingOptions);
+				return FText::AsNumber(ConvertedValue, GetNumberFormattingOptions());
+			}
+
+			FText FormatRangeText(TOptional<FFloatInterval> Range) const
+			{
+				if (!Range.IsSet())
+				{
+					return FText::GetEmpty();
+				}
+
+				// Range size is already relative units, so no need to call ConvertValue.
+				return FText::AsNumber(Range->Size(), GetNumberFormattingOptions());
 			}
 
 		private:
+			static const FNumberFormattingOptions* GetNumberFormattingOptions()
+			{
+				static FNumberFormattingOptions NumberFormattingOptions = FNumberFormattingOptions()
+					.SetMinimumFractionalDigits(1)
+					.SetMaximumFractionalDigits(1);
+
+				return &NumberFormattingOptions;
+			}		
+
 			TAttribute<bool> bRelativeScale;
 			TAttribute<int> Target;
 		};
@@ -160,7 +176,17 @@ namespace AudioWidgets
 				+ SOverlay::Slot()
 				[
 					SNew(STextBlock)
-						.Text_Lambda([InItem, LoudnessValueConverter]() { return LoudnessValueConverter.FormatNumericValueText(InItem->Value.Get()); })
+						.Text_Lambda([InItem, LoudnessValueConverter]()
+							{
+								if (InItem->IsRanged())
+								{
+									return LoudnessValueConverter.FormatRangeText(InItem->Range.Get());
+								}
+								else
+								{
+									return LoudnessValueConverter.FormatNumericValueText(InItem->Value.Get());
+								}
+							})
 						.Font(FStyleDefaults::GetFontInfo(14))
 						.ColorAndOpacity_Lambda(GetValueColorAndOpacity)
 						.Justification(ETextJustify::Right)
@@ -169,7 +195,7 @@ namespace AudioWidgets
 				.VAlign(EVerticalAlignment::VAlign_Bottom)
 				[
 					SNew(STextBlock)
-						.Text_Lambda([bRelativeScale = LoudnessScaleParams.bRelativeScale]() { return (bRelativeScale.Get()) ? INVTEXT("LU") : INVTEXT("LUFS"); })
+						.Text_Lambda([InItem, bRelativeScale = LoudnessScaleParams.bRelativeScale]() { return (InItem->IsRanged() || bRelativeScale.Get()) ? INVTEXT("LU") : INVTEXT("LUFS"); })
 						.Font(FStyleDefaults::GetFontInfo(6))
 						.ColorAndOpacity_Lambda(GetValueColorAndOpacity)
 						.Justification(ETextJustify::Right)
@@ -226,15 +252,40 @@ namespace AudioWidgets
 				+ SVerticalBox::Slot()
 				.HAlign(EHorizontalAlignment::HAlign_Center)
 				[
-					SNew(SVariableRangeMeterWidget)
+					SNew(SVariableDisplayRangeMeterWidget)
 						.MeterValueColor_Lambda([InItem]() { return InItem->Color.Get(); })
 						.MeterChannelInfo_Lambda([InItem, LoudnessValueConverter]() -> TArray<FAudioMeterChannelInfo>
 							{
 								constexpr float DefaultInvalidValue = MIN_VOLUME_DECIBELS;
-								const TOptional<float> ConvertedValue = LoudnessValueConverter.ConvertValue(InItem->Value.Get());
-								return { { ConvertedValue.Get(DefaultInvalidValue), DefaultInvalidValue, DefaultInvalidValue } };
+								if (InItem->IsRanged())
+								{
+									// Display meter bar as a range:
+									const TOptional<FFloatInterval> Range = InItem->Range.Get();
+									const TOptional<float> ConvertedMin = LoudnessValueConverter.ConvertValue(Range.IsSet() ? Range->Min : TOptional<float>());
+									const TOptional<float> ConvertedMax = LoudnessValueConverter.ConvertValue(Range.IsSet() ? Range->Max : TOptional<float>());
+									return
+									{
+										{
+											.MeterValue = ConvertedMax.Get(DefaultInvalidValue),
+											.PeakValue = DefaultInvalidValue,
+											.MeterValueBarStart = ConvertedMin
+										}
+									};
+								}
+								else
+								{
+									// Display standard meter bar:
+									const TOptional<float> ConvertedValue = LoudnessValueConverter.ConvertValue(InItem->Value.Get());
+									return
+									{
+										{
+											.MeterValue = ConvertedValue.Get(DefaultInvalidValue),
+											.PeakValue = DefaultInvalidValue
+										}
+									};
+								}
 							})
-						.Range(LoudnessScaleParams.Range)
+						.DisplayRange(LoudnessScaleParams.Range)
 						.Offset_Lambda([Offset = LoudnessScaleParams.Offset, LoudnessValueConverter] { return LoudnessValueConverter.ConvertValue(Offset.Get()); })
 				]
 				+ SVerticalBox::Slot()
diff --git a/Engine/Plugins/Runtime/AudioWidgets/Source/AudioWidgets/Private/LoudnessMeterWidgetView.h b/Engine/Plugins/Runtime/AudioWidgets/Source/AudioWidgets/Private/LoudnessMeterWidgetView.h
--- a/Engine/Plugins/Runtime/AudioWidgets/Source/AudioWidgets/Private/LoudnessMeterWidgetView.h
+++ b/Engine/Plugins/Runtime/AudioWidgets/Source/AudioWidgets/Private/LoudnessMeterWidgetView.h
@@ -23,6 +23,7 @@ namespace AudioWidgets
 			FName Name;
 			FText DisplayName;
 			TAttribute<TOptional<float>> Value;
+			TAttribute<TOptional<FFloatInterval>> Range;
 			TAttribute<bool> bShowValue = false;
 			TAttribute<bool> bShowMeter = false;
 			FSimpleDelegate OnShowValueToggleRequested;
@@ -31,6 +32,8 @@ namespace AudioWidgets
 			FOnLinearColorValueChanged OnColorChanged;
 			FOnWindowClosed OnColorPickerWindowClosed;
 			FSimpleDelegate OnResetColorRequested;
+
+			bool IsRanged() const { return Range.IsBound() || Range.Get(NullOpt).IsSet(); }
 		};
 
 		struct FTimerPanelParams
diff --git a/Engine/Plugins/Runtime/AudioWidgets/Source/AudioWidgets/Public/LoudnessMeterRackUnitSettings.h b/Engine/Plugins/Runtime/AudioWidgets/Source/AudioWidgets/Public/LoudnessMeterRackUnitSettings.h
--- a/Engine/Plugins/Runtime/AudioWidgets/Source/AudioWidgets/Public/LoudnessMeterRackUnitSettings.h
+++ b/Engine/Plugins/Runtime/AudioWidgets/Source/AudioWidgets/Public/LoudnessMeterRackUnitSettings.h
@@ -48,9 +48,12 @@ struct FLoudnessMeterRackUnitSettings
 	UPROPERTY(EditAnywhere, config, Category = LoudnessMeter)
 	FLoudnessMetricDisplayOptions MomentaryLoudness = { .bShowValue = false, .bShowMeter = true };
 
-	UPROPERTY(EditAnywhere, config, Category = LoudnessMeter, meta = (ClampMin = "0", ClampMax = "5", UIMin = "0", UIMax = "5"))
-	uint64 ValuesOrderingPermutation = 0; // Max == 3! - 1
+	UPROPERTY(EditAnywhere, config, Category = LoudnessMeter)
+	FLoudnessMetricDisplayOptions LoudnessRange = { .bShowValue = true, .bShowMeter = false };
+
+	UPROPERTY(EditAnywhere, config, Category = LoudnessMeter, meta = (ClampMin = "0", ClampMax = "23", UIMin = "0", UIMax = "23"))
+	uint64 ValuesOrderingPermutation = 0; // Max == 4! - 1
 
-	UPROPERTY(EditAnywhere, config, Category = LoudnessMeter, meta = (ClampMin = "0", ClampMax = "5", UIMin = "0", UIMax = "5"))
-	uint64 MetersOrderingPermutation = 0; // Max == 3! - 1
+	UPROPERTY(EditAnywhere, config, Category = LoudnessMeter, meta = (ClampMin = "0", ClampMax = "23", UIMin = "0", UIMax = "23"))
+	uint64 MetersOrderingPermutation = 0; // Max == 4! - 1
 };
diff --git a/Engine/Plugins/Runtime/AudioWidgets/Source/AudioWidgetsCore/Private/SAudioMeterWidget.cpp b/Engine/Plugins/Runtime/AudioWidgets/Source/AudioWidgetsCore/Private/SAudioMeterWidget.cpp
--- a/Engine/Plugins/Runtime/AudioWidgets/Source/AudioWidgetsCore/Private/SAudioMeterWidget.cpp
+++ b/Engine/Plugins/Runtime/AudioWidgets/Source/AudioWidgetsCore/Private/SAudioMeterWidget.cpp
@@ -48,7 +48,7 @@ int32 SAudioMeterWidget::OnPaint(const FPaintArgs& Args, const FGeometry& Allott
  	const int32 DecibelsPerHash = Style->DecibelsPerHash;
  
  	int32 MinValueDb = FMath::Min(static_cast<int32>(Style->ValueRangeDb.X), static_cast<int32>(Style->ValueRangeDb.Y));
- 	int32 MaxValueDb = FMath::Max(static_cast<int32>(Style->ValueRangeDb.Y), static_cast<int32>(Style->ValueRangeDb.Y));
+ 	int32 MaxValueDb = FMath::Max(static_cast<int32>(Style->ValueRangeDb.X), static_cast<int32>(Style->ValueRangeDb.Y));
  
  	const float ValueDeltaDb = MaxValueDb - MinValueDb;
 
@@ -95,16 +95,31 @@ int32 SAudioMeterWidget::OnPaint(const FPaintArgs& Args, const FGeometry& Allott
 	}
 	++LayerId;
 
+	if (ValueDeltaDb == 0.0f)
+	{
+		// ValueRangeDb is zero, don't draw any values.
+		return LayerId;
+	}
+
 	// Draw the meter value
 	for (int32 ChannelIndex = 0; ChannelIndex < NumChannels; ++ChannelIndex)
 	{
 		const float ChannelMeterValueDb = ChannelInfos[ChannelIndex].MeterValue;
 		const float ChannelMeterValue = FMath::Clamp((ChannelMeterValueDb - MinValueDb) / ValueDeltaDb, 0.0f, 1.0f);
 
-		const FVector2D MeterValueTopLeft(Style->MeterPadding.X, 
+		const float ChannelMeterValueBarStartDb = ChannelInfos[ChannelIndex].MeterValueBarStart.Get(MinValueDb);
+		const float ChannelMeterValueBarStart = FMath::Clamp((ChannelMeterValueBarStartDb - MinValueDb) / ValueDeltaDb, 0.0f, 1.0f);
+
+		const FFloatInterval MeterValueBarInterval(
+			FMath::Min(ChannelMeterValueBarStart, ChannelMeterValue),
+			FMath::Max(ChannelMeterValueBarStart, ChannelMeterValue));
+
+		const float AllottedWidthMinusPadding = AllottedWidth - 2.0f * Style->MeterPadding.X;
+
+		const FVector2D MeterValueTopLeft(Style->MeterPadding.X + MeterValueBarInterval.Min * AllottedWidthMinusPadding,
 									Style->MeterPadding.Y + ScaleOffset + Style->MeterValuePadding + ChannelIndex * (Style->MeterSize.Y + Style->MeterPadding.Y));
 		
-		const FVector2D MeterValueSize(ChannelMeterValue * (AllottedWidth - 2.0f * Style->MeterPadding.X), 
+		const FVector2D MeterValueSize(MeterValueBarInterval.Size() * AllottedWidthMinusPadding,
 								 Style->MeterSize.Y - 2.0f * Style->MeterValuePadding);
  
 		FSlateDrawElement::MakeBox(
diff --git a/Engine/Plugins/Runtime/AudioWidgets/Source/AudioWidgetsCore/Public/AudioMeterChannelInfo.h b/Engine/Plugins/Runtime/AudioWidgets/Source/AudioWidgetsCore/Public/AudioMeterChannelInfo.h
--- a/Engine/Plugins/Runtime/AudioWidgets/Source/AudioWidgetsCore/Public/AudioMeterChannelInfo.h
+++ b/Engine/Plugins/Runtime/AudioWidgets/Source/AudioWidgetsCore/Public/AudioMeterChannelInfo.h
@@ -17,11 +17,15 @@ struct FAudioMeterChannelInfo
 
 	UPROPERTY(EditAnywhere, BlueprintReadWrite, Category = AudioMeter)
 	float ClippingValue = 0.0f;
+
+	UPROPERTY(EditAnywhere, BlueprintReadWrite, Category = AudioMeter, meta = (ToolTip = "By default the meter value bar will be drawn from the minimum of the ValueRangeDb up to the MeterValue. Set this property to draw the meter value bar starting at a different value."))
+	TOptional<float> MeterValueBarStart;
 };
 
 inline bool operator==(const FAudioMeterChannelInfo& Lhs, const FAudioMeterChannelInfo& Rhs)
 {
 	return FMath::IsNearlyEqual(Lhs.MeterValue, Rhs.MeterValue) &&
 		   FMath::IsNearlyEqual(Lhs.PeakValue, Rhs.PeakValue)   &&
-		   FMath::IsNearlyEqual(Lhs.ClippingValue, Rhs.ClippingValue);
+		   FMath::IsNearlyEqual(Lhs.ClippingValue, Rhs.ClippingValue) &&
+		   ((!Lhs.MeterValueBarStart.IsSet() && !Rhs.MeterValueBarStart.IsSet()) || (Lhs.MeterValueBarStart.IsSet() && Rhs.MeterValueBarStart.IsSet() && FMath::IsNearlyEqual(*Lhs.MeterValueBarStart, *Rhs.MeterValueBarStart)));
 }
