[AudioSynesthesia] [SignalProcessing] Refactor to move LKFS calculations from AudioSynesthesiaCore to SignalProcessing so LKFS analysis can be done in other parts of the Engine/Modules that can't depend on AudioSynesthesia.

#Synced-CL 51006836

diff --git a/Engine/Plugins/Runtime/AudioSynesthesia/Source/AudioSynesthesiaCore/Private/LKFSFactory.cpp b/Engine/Plugins/Runtime/AudioSynesthesia/Source/AudioSynesthesiaCore/Private/LKFSFactory.cpp
--- a/Engine/Plugins/Runtime/AudioSynesthesia/Source/AudioSynesthesiaCore/Private/LKFSFactory.cpp
+++ b/Engine/Plugins/Runtime/AudioSynesthesia/Source/AudioSynesthesiaCore/Private/LKFSFactory.cpp
@@ -2,160 +2,34 @@
 
 #include "LKFSFactory.h"
 
-#include "Algo/IsSorted.h"
-#include "AudioSynesthesiaCoreLog.h"
-#include "Containers/Array.h"
-#include "DSP/SlidingWindow.h"
-#include "LKFSPrivate.h"
-
 #include UE_INLINE_GENERATED_CPP_BY_NAME(LKFSFactory)
 
 namespace Audio
 {
 	namespace LKFSFactoryPrivate
 	{
-		float LKFSToEnergy(float InLKFS)
-		{
-			return FMath::Pow(10, (InLKFS + 0.691f) / 10.f);
-		}
-
-		// Calculates average value from a circular buffer. Allows for interleaved data
-		// by providing a "stride". 
-		float CalculateAverageFromHistory(const TCircularBuffer<float>& InHistory, int32 InStartIndex, int32 InEndIndex, int32 InStride)
-		{
-			float Accum = 0.f;
-			for (int32 Index = InStartIndex; Index < InEndIndex; Index += InStride)
-			{
-				Accum += InHistory[Index];
-			}
-			int32 Num = (InEndIndex - InStartIndex) / InStride;
-			if (Num > 0)
-			{
-				return Accum /= Num;
-			}
-			return 0.f;
-		}
-
-		// Calculates gated average value from a circular buffer. Allows for interleaved data
-		// by providing a "stride". 
-		float CalculateGatedAverageFromHistory(const TCircularBuffer<float>& InHistory, int32 InStartIndex, int32 InEndIndex, int32 InStride, float InGate)
-		{
-			float Accum = 0.f;
-			int32 Num = 0;
-			for (int32 Index = InStartIndex; Index < InEndIndex; Index += InStride)
-			{
-				float Val = InHistory[Index];
-				if (Val > InGate)
-				{
-					Accum += Val;
-					Num++;
-				}
-			}
-			if (Num > 0)
-			{
-				return Accum /= Num;
-			}
-			return InGate;
-		}
-
-		float CalculateGatedLoudnessFromHistory(const TCircularBuffer<float>& InHistory, int32 InStartIndex, int32 InEndIndex, int32 InStride)
-		{
-			// Gated Loudness only examines loudnesses above a certain gating loudness. 
-			//
-			// It is a 2 pass algorithm 
-			// Pass 1: Use a set loudness threshold of -70lkfs to determine an initial gated loudness
-			// Pass 2: Use the initial gated loudness to derive a new loudness threshold and calculate the gate loudness again. 
-			constexpr float InitLoudnessThreshold = -70.f;
-			float InitEnergyThreshold = LKFSToEnergy(InitLoudnessThreshold);
-
-			const float GatedEnergyPass1 = CalculateGatedAverageFromHistory(InHistory, InStartIndex, InEndIndex, InStride, InitEnergyThreshold);
-
-			const float NewThreshold = GatedEnergyPass1 * 0.1f;
-			const float GatedEnergyPass2 = CalculateGatedAverageFromHistory(InHistory, InStartIndex, InEndIndex, InStride, NewThreshold);
-			return FLoudnessAnalyzer::ConvertPerceptualEnergyToLKFS(GatedEnergyPass2);
-		}
-
-		/**
-		 * Implementation copied from FFineGrainedPerformanceTracker::GetPercentileValue(...)
-		 */
-		float GetPercentileValue(TArray<float>& Samples, int32 Percentile)
+		FLKFSAnalyzer MakeLKFSAnalyzer(const FAnalyzerParameters& InParams, const FLKFSSettings& InAnalyzerSettings)
 		{
-			ensure(Percentile >= 0 && Percentile <= 100);
-
-			int32 Left = 0;
-			int32 Right = Samples.Num() - 1;
-
-			if (Right < 0)
+			const FLKFSAnalyzerSettings LKFSAnalyzerSettings
 			{
-				return -1.0f;
-			}
-
-			const int32 PercentileClamped = FMath::Clamp(Percentile, 0, 100);
-			const int32 PercentileOrdinal = FMath::RoundToInt32((float(PercentileClamped) * Right) / 100);
-
-			// this is quickselect (see http://en.wikipedia.org/wiki/Quickselect for details).
-			while (Right != Left)
-			{
-				// partition
-				int32 MovingLeft = Left - 1;
-				int32 MovingRight = Right;
-				float Pivot = Samples[MovingRight];
-				for (; ; )
-				{
-					while (Samples[++MovingLeft] < Pivot);
-					while (Samples[--MovingRight] > Pivot)
-					{
-						if (MovingRight == Left)
-						{
-							break;
-						}
-					}
-
-					if (MovingLeft >= MovingRight)
-					{
-						break;
-					}
-
-					const float Temp = Samples[MovingLeft];
-					Samples[MovingLeft] = Samples[MovingRight];
-					Samples[MovingRight] = Temp;
-				}
-
-				const float Temp = Samples[MovingLeft];
-				Samples[MovingLeft] = Samples[Right];
-				Samples[Right] = Temp;
-
-				// now we're pivoted around MovingLeft
-				// decide what part K-th largest belong to
-				if (MovingLeft > PercentileOrdinal)
-				{
-					Right = MovingLeft - 1;
-				}
-				else if (MovingLeft < PercentileOrdinal)
-				{
-					Left = MovingLeft + 1;
-				}
-				else
-				{
-					// we hit exactly the value we need, no need to sort further
-					break;
-				}
-			}
+				.LoudnessAnalyzerSettings          = InAnalyzerSettings.LoudnessAnalyzerSettings,
+				.AnalysisPeriod                    = InAnalyzerSettings.AnalysisPeriod,
+				.AnalysisWindowDuration            = InAnalyzerSettings.AnalysisWindowDuration,
+				.ShortTermLoudnessDuration         = InAnalyzerSettings.ShortTermLoudnessDuration,
+				.IntegratedLoudnessAnalysisPeriod  = InAnalyzerSettings.IntegratedLoudnessAnalysisPeriod,
+				.IntegratedLoudnessDuration        = InAnalyzerSettings.IntegratedLoudnessDuration,
+				.bCalculatePerChannelLoudnessRange = InAnalyzerSettings.bCalculatePerChannelLoudnessRange,
+				.bCalculateOverallLoudnessRange    = InAnalyzerSettings.bCalculateOverallLoudnessRange
+			};
 
-			return Samples[PercentileOrdinal];
+			return FLKFSAnalyzer(InParams.NumChannels, static_cast<float>(InParams.SampleRate), LKFSAnalyzerSettings);
 		}
-	}
-}
+	} // namespace LKFSFactoryPrivate
 
-namespace Audio
-{
-	// Use for overall loudness storage
-	const int32 FLKFSAnalyzerResult::ChannelIndexOverall = -1;
-
-	FLKFSAnalyzerResult::FLKFSAnalyzerResult()
-	{
-	}
 
+	//----------------------------------------------------------------------------------------
+	// FLKFSAnalyzerResult
+	//----------------------------------------------------------------------------------------
 	int32 FLKFSAnalyzerResult::GetNumChannels() const
 	{
 		int32 Num = ChannelResults.Num();
@@ -190,324 +64,69 @@ namespace Audio
 		return GetChannelLoudnessResults(ChannelIndexOverall);
 	}
 
+
+	//----------------------------------------------------------------------------------------
+	// FLKFSWorker
+	//----------------------------------------------------------------------------------------
 	FLKFSWorker::FLKFSWorker(const FAnalyzerParameters& InParams, const FLKFSSettings& InAnalyzerSettings)
-	: NumChannels(InParams.NumChannels)
-	, SampleRate(InParams.SampleRate)
-	, ChannelEnergyHistory(1)
-	, OverallEnergyHistory(1)
-	, Analyzer(InParams.SampleRate, InAnalyzerSettings.LoudnessAnalyzerSettings)
+		: NumChannels(InParams.NumChannels)
+		, LKFSAnalyzer(LKFSFactoryPrivate::MakeLKFSAnalyzer(InParams, InAnalyzerSettings))
 	{
-		check(NumChannels > 0);
-		check(SampleRate > 0.f);
-		check(InAnalyzerSettings.ShortTermLoudnessDuration > 0.f);
-		check(InAnalyzerSettings.IntegratedLoudnessDuration > 0.f);
-		check(InAnalyzerSettings.IntegratedLoudnessAnalysisPeriod> 0.f);
-
-
-		int32 NumSlidingWindowFrames = Analyzer.GetSettings().WindowSize;
-
-		int32 NumAnalysisFrames = FMath::CeilToInt(InAnalyzerSettings.AnalysisWindowDuration * InParams.SampleRate);
-		int32 NumAnalysisHopFrames = FMath::CeilToInt(InAnalyzerSettings.AnalysisPeriod * InParams.SampleRate);
-
-		// There may be multiple calls to the LoudnessAnalyzer to produce a single 
-		// FLoudnessDatum. The sliding window hop is tuned here so that it best 
-		// matches the desired AnalysisPeriod while maintaining between a 25% to 
-		// 75% window overlap.
-		LKFSPrivate::TuneSlidingWindwoHopSize(NumAnalysisHopFrames, NumSlidingWindowFrames, NumSlidingWindowHopFrames, NumAnalysisHopWindows);
-
-		// Determine how many windows to analyze per loudness result
-		NumAnalysisWindows = FMath::Max(1, ((NumAnalysisFrames - (NumSlidingWindowFrames / 2)) / NumSlidingWindowHopFrames) + 1);
-
-		// Include NumChannels in calculating window size because windows are generated 
-		// with interleaved samples and deinterleaved later. 
-		int32 NumSlidingWindowSamples = NumSlidingWindowFrames * NumChannels;
-		int32 NumSlidingWindowHopSamples = NumSlidingWindowHopFrames * NumChannels;
-
-		InternalBuffer = TSlidingBuffer<float>(NumSlidingWindowSamples, NumSlidingWindowHopSamples);
-		InternalWindow.AddUninitialized(NumSlidingWindowSamples);
-		ChannelPerceptualEnergy.AddZeroed(NumChannels * NumAnalysisWindows);
-		OverallPerceptualEnergy.AddZeroed(NumAnalysisWindows);
-
-		// Initialize settings for calculating aggregate stats
-		const int32 NumResultHopFrames = NumAnalysisHopWindows * NumSlidingWindowHopFrames;
-		NumShortTermHistoryResults = FMath::Max(1, FMath::DivideAndRoundUp((int32) (InAnalyzerSettings.ShortTermLoudnessDuration * SampleRate), NumResultHopFrames));
-		NumIntegratedHistoryResults = FMath::Max(1, FMath::DivideAndRoundUp((int32) (InAnalyzerSettings.IntegratedLoudnessDuration * SampleRate), NumResultHopFrames));
-		NumIntegratedHopHistoryResults = FMath::Max(1, FMath::DivideAndRoundUp((int32) (InAnalyzerSettings.IntegratedLoudnessAnalysisPeriod * SampleRate), NumResultHopFrames));
-
-		const int32 MaxHistory = FMath::Max(NumIntegratedHistoryResults, NumShortTermHistoryResults);
-
-		ChannelEnergyHistory = TCircularBuffer<float>(MaxHistory * NumChannels);
-		OverallEnergyHistory = TCircularBuffer<float>(MaxHistory);
-
-		for (int32 ChannelIdx = 0; ChannelIdx < NumChannels; ChannelIdx++)
-		{
-			AggregateLoudnessStats.Add(ChannelIdx);
-		}
-		AggregateLoudnessStats.Add(FLKFSAnalyzerResult::ChannelIndexOverall);
-
-		if (InAnalyzerSettings.bCalculatePerChannelLoudnessRange)
-		{
-			for (int32 ChannelIdx = 0; ChannelIdx < NumChannels; ChannelIdx++)
-			{
-				LoudnessRangeStatePerChannel.Emplace(ChannelIdx, static_cast<uint32>(NumIntegratedHistoryResults));
-			}
-		}
-		if (InAnalyzerSettings.bCalculateOverallLoudnessRange)
-		{
-			LoudnessRangeStatePerChannel.Emplace(FLKFSAnalyzerResult::ChannelIndexOverall, static_cast<uint32>(NumIntegratedHistoryResults));
-		}
 	}
 
 	void FLKFSWorker::Analyze(TArrayView<const float> InAudio, IAnalyzerResult* OutResult) 
 	{
 		// Assume that outer layers ensured that this is of correct type.
 		FLKFSAnalyzerResult* LoudnessResult = static_cast<FLKFSAnalyzerResult*>(OutResult);
-
-		check(nullptr != LoudnessResult);
+		check(LoudnessResult != nullptr);
 		
-		TAutoSlidingWindow<float> SlidingWindow(InternalBuffer, InAudio, InternalWindow);
+		FLKFSAnalyzerResults LKFSAnalyzerResults;
+		LKFSAnalyzer.Analyze(InAudio, LKFSAnalyzerResults);
 
-		AnalyzeSlidingWindow(SlidingWindow, *LoudnessResult);
-	}
-
-	void FLKFSWorker::AnalyzeSlidingWindow(TAutoSlidingWindow<float>& InSlidingWindow, FLKFSAnalyzerResult& OutResult)
-	{
-		check(NumAnalysisWindows > 0);
-		check(NumAnalysisHopWindows > 0);
-		
-		for (const TArray<float>& Window : InSlidingWindow)
+		// Per Channel results
+		for (int32 Index = 0; Index < NumChannels; ++Index)
 		{
-			AnalyzeWindow(Window);
-
-			if ((NumAnalyzedBuffers >= NumAnalysisWindows) && ((NumAnalyzedBuffers - NumAnalysisWindows) % NumAnalysisHopWindows) == 0)
+			for (const FLKFSResult& ChannelLoudnessResult : LKFSAnalyzerResults.GetChannelLoudnessResults(Index))
 			{
-				AddCurrentLoudnessDataToResult(OutResult);
-			}
-		}
-	}
-
-	void FLKFSWorker::AnalyzeWindow(TArrayView<const float> InWindow)
-	{
-		const int32 WrappedWindowIdx = (NumAnalyzedBuffers % NumAnalysisWindows);
-		const int32 ChannelPerceptualEnergyWritePos = WrappedWindowIdx * NumChannels;
-
-		// Calculate perceptual energy
-		OverallPerceptualEnergy[WrappedWindowIdx] = Analyzer.CalculatePerceptualEnergy(InWindow, NumChannels, MakeArrayView<float>(ChannelPerceptualEnergy.GetData() + ChannelPerceptualEnergyWritePos, NumChannels));
-
-		// update counters
-		NumAnalyzedBuffers++;
-	}
-
-	void FLKFSWorker::AddCurrentLoudnessDataToHistory()
-	{
-		// Compute loudness per channel and store history
-		for (int32 ChannelIdx = 0; ChannelIdx < NumChannels; ChannelIdx++)
-		{
-			// Determine average perceptual energy over analysis window
-			float PerceptualEnergy = 0.f;
-			for (int32 Pos = ChannelIdx; Pos < ChannelPerceptualEnergy.Num(); Pos += NumChannels)
-			{
-				PerceptualEnergy += ChannelPerceptualEnergy[Pos];
-			}
-			PerceptualEnergy /= (float)NumAnalysisWindows;
-
-			// Store perceptual energy in history 
-			ChannelEnergyHistory[(NumChannels * NumEnergyHistory) + ChannelIdx] = PerceptualEnergy;
-		}
-
-		{
-			// Compute loudness overall
-			float PerceptualEnergy = 0.f;
-			for (int32 Pos = 0; Pos < OverallPerceptualEnergy.Num(); Pos++)
-			{
-				PerceptualEnergy += OverallPerceptualEnergy[Pos];
+				LoudnessResult->Add(
+					{
+						.Channel            = ChannelLoudnessResult.Channel,
+						.Timestamp          = ChannelLoudnessResult.Timestamp,
+						.Energy             = ChannelLoudnessResult.Energy,
+						.Loudness           = ChannelLoudnessResult.Loudness,
+						.ShortTermLoudness  = ChannelLoudnessResult.ShortTermLoudness,
+						.IntegratedLoudness = ChannelLoudnessResult.IntegratedLoudness,
+						.GatedLoudness      = ChannelLoudnessResult.GatedLoudness,
+						.LoudnessRange      = ChannelLoudnessResult.LoudnessRange
+					}
+				);
 			}
-			PerceptualEnergy /= (float)NumAnalysisWindows;
-
-			OverallEnergyHistory[NumEnergyHistory] = PerceptualEnergy;
 		}
-	
-		NumEnergyHistory++;
-	}
-
-	void FLKFSWorker::AddCurrentLoudnessDataToResult(FLKFSAnalyzerResult& OutResult)
-	{
-		AddCurrentLoudnessDataToHistory();
 
-		UpdateAggregateLoudnessStats();
-
-		// Accounting for all the frames positions analyzed in creating the 
-		// data, determine the center frame. 
-		const int32 CenterFramePos = ((2 * NumAnalyzedBuffers - NumAnalysisWindows) * NumSlidingWindowHopFrames + Analyzer.GetSettings().WindowSize) / 2;
-		const float Timestamp = CenterFramePos / SampleRate;
-
-
-		// Create results and store
-		for (int32 ChannelIdx = 0; ChannelIdx < NumChannels; ChannelIdx++)
-		{
-			int32 HistoryIndex = ((NumEnergyHistory - 1) * NumChannels) + ChannelIdx;
-			
-			float Energy = ChannelEnergyHistory[HistoryIndex];
-			const FAggregateLoudnessStats& Stats = AggregateLoudnessStats[ChannelIdx];
-
-			OutResult.Add(
-				FLKFSResults 
-				{
-					.Channel=ChannelIdx, 
-					.Timestamp=Timestamp, 
-					.Energy=Energy, 
-					.Loudness=FLoudnessAnalyzer::ConvertPerceptualEnergyToLKFS(Energy),
-					.ShortTermLoudness=Stats.ShortTermLoudness, 
-					.IntegratedLoudness=Stats.IntegratedLoudness,
-					.GatedLoudness=Stats.GatedLoudness,
-					.LoudnessRange=Stats.LoudnessRange,
-				}
-			);
-		}
+		// Overall results
+		for (const FLKFSResult& OverallLoudnessResult : LKFSAnalyzerResults.GetLoudnessResults())
 		{
-			int32 HistoryIndex = (NumEnergyHistory - 1);
-			
-			float Energy = OverallEnergyHistory[HistoryIndex];
-			const FAggregateLoudnessStats& Stats = AggregateLoudnessStats[FLKFSAnalyzerResult::ChannelIndexOverall];
-
-			OutResult.Add(
-				FLKFSResults 
+			LoudnessResult->Add(
 				{
-					.Channel=FLKFSAnalyzerResult::ChannelIndexOverall, 
-					.Timestamp=Timestamp, 
-					.Energy=Energy, 
-					.Loudness=FLoudnessAnalyzer::ConvertPerceptualEnergyToLKFS(Energy),
-					.ShortTermLoudness=Stats.ShortTermLoudness, 
-					.IntegratedLoudness=Stats.IntegratedLoudness,
-					.GatedLoudness=Stats.GatedLoudness,
-					.LoudnessRange=Stats.LoudnessRange,
+					.Channel            = OverallLoudnessResult.Channel,
+					.Timestamp          = OverallLoudnessResult.Timestamp,
+					.Energy             = OverallLoudnessResult.Energy,
+					.Loudness           = OverallLoudnessResult.Loudness,
+					.ShortTermLoudness  = OverallLoudnessResult.ShortTermLoudness,
+					.IntegratedLoudness = OverallLoudnessResult.IntegratedLoudness,
+					.GatedLoudness      = OverallLoudnessResult.GatedLoudness,
+					.LoudnessRange      = OverallLoudnessResult.LoudnessRange
 				}
 			);
 		}
 	}
 
-	void FLKFSWorker::UpdateAggregateLoudnessStats()
-	{
-		// Short term loudness is always updated, but integrated & gated loudness are 
-		// updated at a lower rate. 
-		bool bUpdateIntegratedLoudness = (NumEnergyHistory == 1) || (NumEnergyHistory % NumIntegratedHopHistoryResults) == 0;
-
-
-		// Per channel results
-		for (int32 ChannelIdx = 0; ChannelIdx < NumChannels; ChannelIdx++)
-		{
-			FAggregateLoudnessStats& Stats = AggregateLoudnessStats[ChannelIdx];
-			FLoudnessRange* LoudnessRange = LoudnessRangeStatePerChannel.Find(ChannelIdx);
-			UpdateAggregateLoudnessStats(ChannelEnergyHistory, ChannelIdx, NumChannels, bUpdateIntegratedLoudness, Stats, LoudnessRange);
-		}
-
-		// Overall results
-		{
-			FAggregateLoudnessStats& Stats = AggregateLoudnessStats[FLKFSAnalyzerResult::ChannelIndexOverall];
-			FLoudnessRange* LoudnessRange = LoudnessRangeStatePerChannel.Find(FLKFSAnalyzerResult::ChannelIndexOverall);
-			UpdateAggregateLoudnessStats(OverallEnergyHistory, 0 /*offset*/, 1/*stride*/, bUpdateIntegratedLoudness, Stats, LoudnessRange);
-		}
-	}
-
-	void FLKFSWorker::UpdateAggregateLoudnessStats(const TCircularBuffer<float>& InHistory, int32 InOffset, int32 InStride, bool bInUpdateIntegratedLoudness, FAggregateLoudnessStats& OutStats, FLoudnessRange* LoudnessRange) const
-	{
-		const int32 EndIndex = (NumEnergyHistory * InStride) + InOffset;
-
-		// Update short term loudness always.
-		const int32 ShortTermLoudnessStartIndex = InOffset + FMath::Max(0, (NumEnergyHistory - NumShortTermHistoryResults) * InStride);
-		const float ShortTermAverageEnergy = LKFSFactoryPrivate::CalculateAverageFromHistory(InHistory, ShortTermLoudnessStartIndex, EndIndex, InStride);
-		OutStats.ShortTermLoudness = FLoudnessAnalyzer::ConvertPerceptualEnergyToLKFS(ShortTermAverageEnergy);
-
-		// Only update long term loudness if flag is set. 
-		if (bInUpdateIntegratedLoudness)
-		{
-			const int32 LongTermLoudnessStartIndex = InOffset + FMath::Max(0, (NumEnergyHistory - NumIntegratedHistoryResults * InStride));
-
-			float IntegratedEnergy = LKFSFactoryPrivate::CalculateAverageFromHistory(InHistory, LongTermLoudnessStartIndex, EndIndex, InStride);
-			OutStats.IntegratedLoudness = FLoudnessAnalyzer::ConvertPerceptualEnergyToLKFS(IntegratedEnergy);
-			OutStats.GatedLoudness = LKFSFactoryPrivate::CalculateGatedLoudnessFromHistory(InHistory, LongTermLoudnessStartIndex, EndIndex, InStride);
-		}
-
-		// Update optional loudness range (if we have LoudnessRange state):
-		if (LoudnessRange)
-		{
-			LoudnessRange->AddShortTermEnergyMeasurement(ShortTermAverageEnergy);
-
-			if (bInUpdateIntegratedLoudness)
-			{
-				OutStats.LoudnessRange = LoudnessRange->CalculateLoudnessRange();
-			}
-		}
-	}
-
-	void FLKFSWorker::FLoudnessRange::AddShortTermEnergyMeasurement(float InShortTermEnergy)
-	{
-		ShortTermEnergyHistory[ShortTermWriteIndex++] = InShortTermEnergy;
-		ShortTermHistorySize = FMath::Min<uint32>(ShortTermHistorySize + 1, MaxShortTermHistorySize);
-	}
-
-	TOptional<FFloatInterval> FLKFSWorker::FLoudnessRange::CalculateLoudnessRange() const
-	{
-		// Based on EBU Tech 3342
-		if (ShortTermHistorySize == 0)
-		{
-			return NullOpt;
-		}
-
-		constexpr float AbsoluteThresholdEnergy = 1.17246636e-7f; // -70 LUFS
-		constexpr float RelativeThresholdEnergy = 0.01f;		  // -20 LU 
-
-		const uint32 LoudnessRangeStartIndex = ShortTermWriteIndex - ShortTermHistorySize;
-		const uint32 LoudnessRangeEndIndex   = ShortTermWriteIndex;
-
-		// Apply the absolute-threshold gating
-		const float RefEnergy = LKFSFactoryPrivate::CalculateGatedAverageFromHistory(
-			ShortTermEnergyHistory,
-			static_cast<int32>(LoudnessRangeStartIndex),
-			static_cast<int32>(LoudnessRangeEndIndex),
-			1 /*Stride*/,
-			AbsoluteThresholdEnergy
-		);
-
-		// Apply the relative-threshold gating
-		const float RelGateEnergy = FMath::Max(RefEnergy * RelativeThresholdEnergy, AbsoluteThresholdEnergy);
-
-		TArray<float> ShortTermLoudnessRelGated;
-		ShortTermLoudnessRelGated.Reserve(ShortTermHistorySize);
-
-		for (uint32 Index = LoudnessRangeStartIndex; Index < LoudnessRangeEndIndex; ++Index)
-		{
-			const float ShortTermEnergy = ShortTermEnergyHistory[Index];
-
-			if (ShortTermEnergy >= RelGateEnergy)
-			{
-				ShortTermLoudnessRelGated.Emplace(FLoudnessAnalyzer::ConvertPerceptualEnergyToLKFS(ShortTermEnergy));
-			}
-		}
-
-		if (ShortTermLoudnessRelGated.IsEmpty())
-		{
-			return NullOpt;
-		}
-
-		// Compute the high and low percentiles of the distribution of values in ShortTermLoudnessRelGated
-		static constexpr int32 LowerPercentile = 10;
-		static constexpr int32 UpperPercentile = 95;
-
-		const float LoudnessRangeMin = LKFSFactoryPrivate::GetPercentileValue(ShortTermLoudnessRelGated, LowerPercentile);
-		const float LoudnessRangeMax = LKFSFactoryPrivate::GetPercentileValue(ShortTermLoudnessRelGated, UpperPercentile);
-
-		return FFloatInterval(LoudnessRangeMin, LoudnessRangeMax);
-	}
-
-	/** Name of specific analyzer type. */
 	FName FLKFSFactory::GetName() const 
 	{
 		static FName FactoryName(TEXT("LKFSFactory"));
 		return FactoryName;
 	}
 
-	/** Human readable name of analyzer. */
 	FString FLKFSFactory::GetTitle() const
 	{
 		return TEXT("LKFS Analyzer Real-Time");
@@ -527,5 +146,4 @@ namespace Audio
 
 		return MakeUnique<FLKFSWorker>(InParams, *LoudnessSettings);
 	}
-}
-
+} // namespace Audio
diff --git a/Engine/Plugins/Runtime/AudioSynesthesia/Source/AudioSynesthesiaCore/Private/LKFSNRTFactory.cpp b/Engine/Plugins/Runtime/AudioSynesthesia/Source/AudioSynesthesiaCore/Private/LKFSNRTFactory.cpp
--- a/Engine/Plugins/Runtime/AudioSynesthesia/Source/AudioSynesthesiaCore/Private/LKFSNRTFactory.cpp
+++ b/Engine/Plugins/Runtime/AudioSynesthesia/Source/AudioSynesthesiaCore/Private/LKFSNRTFactory.cpp
@@ -4,8 +4,8 @@
 
 #include "Algo/IsSorted.h"
 #include "AudioSynesthesiaCoreLog.h"
+#include "DSP/LKFSUtils.h"
 #include "DSP/SlidingWindow.h"
-#include "LKFSPrivate.h"
 
 #include UE_INLINE_GENERATED_CPP_BY_NAME(LKFSNRTFactory)
 
@@ -254,10 +254,10 @@ namespace Audio
 		int32 NumAnalysisHopFrames = FMath::CeilToInt(InAnalyzerSettings.AnalysisPeriod * InParams.SampleRate);
 
 		// There may be multiple calls to the LoudnessAnalyzer to produce a single 
-		// FLoundessDatum. The sliding window hop is tuned here so that it best 
+		// FLoudnessDatum. The sliding window hop is tuned here so that it best 
 		// matches the desired AnalysisPeriod while maintaining between a 25% to 
 		// 75% window overlap.
-		LKFSPrivate::TuneSlidingWindwoHopSize(NumAnalysisHopFrames, NumSlidingWindowFrames, NumSlidingWindowHopFrames, NumAnalysisHopWindows);
+		LKFSUtils::TuneSlidingWindowHopSize(NumAnalysisHopFrames, NumSlidingWindowFrames, NumSlidingWindowHopFrames, NumAnalysisHopWindows);
 
 		// Determine how many windows to analyze per loudness result
 		NumAnalysisWindows = FMath::Max(1, ((NumAnalysisFrames - (NumSlidingWindowFrames / 2)) / NumSlidingWindowHopFrames) + 1);
diff --git a/Engine/Plugins/Runtime/AudioSynesthesia/Source/AudioSynesthesiaCore/Private/LKFSPrivate.cpp b/Engine/Plugins/Runtime/AudioSynesthesia/Source/AudioSynesthesiaCore/Private/LKFSPrivate.cpp
deleted file mode 100644
diff --git a/Engine/Plugins/Runtime/AudioSynesthesia/Source/AudioSynesthesiaCore/Private/LKFSPrivate.h b/Engine/Plugins/Runtime/AudioSynesthesia/Source/AudioSynesthesiaCore/Private/LKFSPrivate.h
deleted file mode 100644
diff --git a/Engine/Plugins/Runtime/AudioSynesthesia/Source/AudioSynesthesiaCore/Private/LoudnessAnalyzer.cpp b/Engine/Plugins/Runtime/AudioSynesthesia/Source/AudioSynesthesiaCore/Private/LoudnessAnalyzer.cpp
deleted file mode 100644
diff --git a/Engine/Plugins/Runtime/AudioSynesthesia/Source/AudioSynesthesiaCore/Public/LKFSFactory.h b/Engine/Plugins/Runtime/AudioSynesthesia/Source/AudioSynesthesiaCore/Public/LKFSFactory.h
--- a/Engine/Plugins/Runtime/AudioSynesthesia/Source/AudioSynesthesiaCore/Public/LKFSFactory.h
+++ b/Engine/Plugins/Runtime/AudioSynesthesia/Source/AudioSynesthesiaCore/Public/LKFSFactory.h
@@ -3,10 +3,8 @@
 #pragma once
 
 #include "Containers/ArrayView.h"
-#include "Containers/CircularBuffer.h"
 #include "Containers/SortedMap.h"
-#include "DSP/SlidingWindow.h"
-#include "LoudnessAnalyzer.h"
+#include "DSP/LKFSAnalyzer.h"
 #include "IAudioAnalyzerInterface.h"
 
 #include "LKFSFactory.generated.h"
@@ -59,48 +57,8 @@ namespace Audio
 	/**
 	 * Contains settings for loudness analyzer.
 	 */
-	class FLKFSSettings : public IAnalyzerSettings 
+	class FLKFSSettings : public IAnalyzerSettings, public FLKFSAnalyzerSettings
 	{
-	public:
-		// These values are taken from the suggested standards for LKFS/LUFS defined
-		// in ITU-R BS.1770
-		static constexpr float StandardShortTermLoudnessDuration = 3.f; // 3 seconds
-		static constexpr float StandardIntegratedLoudnessAnalysisPeriod = 1.f; // 1 second
-																			   //
-		// Streaming audio does not have a natural end point. Integrated loudness
-		// in this implementation is taken over a long duration (as opposed to 
-		// accumulating indefinitely). 
-		static constexpr float DefaultIntegratedLoudnessDuration = 60.f; // 1 minute
-
-		FLoudnessAnalyzerSettings LoudnessAnalyzerSettings = { 
-			.FFTSize = 2048,
-			.WindowType = EWindowType::Blackman,
-			.LoudnessCurveType = ELoudnessCurveType::K,
-			.MinAnalysisFrequency = 20.f,
-			.MaxAnalysisFrequency = 20000.f,
-			.ScalingMethod = ELoudnessAnalyzerScalingMethod::Corrected
-		};
-
-		/** Number of seconds between loudness measurements */
-		float AnalysisPeriod = 0.1f;
-
-		/** Number of seconds of audio analyzed for each loudness measurements */
-		float AnalysisWindowDuration = 0.4f;
-
-		/** Duration of audio analyzed for short term loudness. */
-		float ShortTermLoudnessDuration = StandardShortTermLoudnessDuration;
-
-		/** Number of seconds between long term loudness updates. */
-		float IntegratedLoudnessAnalysisPeriod = StandardIntegratedLoudnessAnalysisPeriod;
-
-		/** Duration of audio analyzed for long term loudness. */
-		float IntegratedLoudnessDuration = DefaultIntegratedLoudnessDuration;
-
-		/** Should the Loudness Range (LRA) statistic be calculated for each channel? */
-		bool bCalculatePerChannelLoudnessRange = false;
-
-		/** Should the overall Loudness Range (LRA) statistic be calculated? */
-		bool bCalculateOverallLoudnessRange = false;
 	};
 
 
@@ -109,9 +67,9 @@ namespace Audio
 	{
 	public:
 		/** Denotes the overall loudness channel index as opposed individual channel indices. */
-		AUDIOSYNESTHESIACORE_API static const int32 ChannelIndexOverall;
+		inline static const int32 ChannelIndexOverall = -1;
 
-		AUDIOSYNESTHESIACORE_API FLKFSAnalyzerResult();
+		FLKFSAnalyzerResult() = default;
 
 		/** The number of individual channels available in the result. */
 		AUDIOSYNESTHESIACORE_API int32 GetNumChannels() const;
@@ -144,83 +102,8 @@ namespace Audio
 		AUDIOSYNESTHESIACORE_API virtual void Analyze(TArrayView<const float> InAudio, IAnalyzerResult* OutResult) override;
 
 	private:
-		struct FAggregateLoudnessStats
-		{
-			static constexpr float InvalidLoudness = TNumericLimits<float>::Lowest();
-
-			float ShortTermLoudness = InvalidLoudness;
-			float IntegratedLoudness = InvalidLoudness;
-			float GatedLoudness = InvalidLoudness;
-			TOptional<FFloatInterval> LoudnessRange;
-		};
-
-		/**
-		 * Holds state for Loudness Range calculation (for a single channel).
-		 */
-		class FLoudnessRange
-		{
-		public:
-			explicit FLoudnessRange(uint32 InMaxShortTermHistorySize)
-				: ShortTermEnergyHistory(InMaxShortTermHistorySize)
-				, MaxShortTermHistorySize(InMaxShortTermHistorySize)
-			{
-			}
-			
-			void AddShortTermEnergyMeasurement(float InShortTermEnergy);
-
-			TOptional<FFloatInterval> CalculateLoudnessRange() const;
-
-		private:
-			TCircularBuffer<float> ShortTermEnergyHistory;
-
-			uint32 MaxShortTermHistorySize = 0;
-			uint32 ShortTermHistorySize = 0;
-			uint32 ShortTermWriteIndex  = 0;
-		};
-
-		void TuneSlidingWindwoHopSize(const int32 InAnalysisHopFrames, const int32 InWindowFrames, int32& OutBestWindowHopFrames, int32& OutBestAnalysisHopWindows) const;
-		void AnalyzeSlidingWindow(TAutoSlidingWindow<float>& InSlidingWindow, FLKFSAnalyzerResult& OutResult);
-		void AnalyzeWindow(TArrayView<const float> InWindow);
-
-		void AddCurrentLoudnessDataToHistory();
-		void AddCurrentLoudnessDataToResult(FLKFSAnalyzerResult& OutResult);
-		void UpdateAggregateLoudnessStats();
-		void UpdateAggregateLoudnessStats(const TCircularBuffer<float>& InHistory, int32 InOffset, int32 InStride, bool bInUpdateIntegratedLoudness, FAggregateLoudnessStats& OutStats, FLoudnessRange* LoudnessRange = nullptr) const;
-
-
 		int32 NumChannels = 0;
-		int32 NumAnalyzedBuffers = 0;
-		int32 NumEnergyHistory = 0;
-		int32 NumSlidingWindowHopFrames = 0;
-		float SampleRate = 0.f;
-
-		// Number of windows between each loudness analysis
-		int32 NumAnalysisHopWindows = 1;
-		// Number of windows inspected for each loudness analysis
-		int32 NumAnalysisWindows = 1;
-		// Number of values in the "EnergyHistory" to analyze for short term loudness
-		int32 NumShortTermHistoryResults = 1;
-		// Number of values in the "EnergyHistory" to analyze for integrated and gated loudness
-		int32 NumIntegratedHistoryResults= 1;
-		// Number of values in the "EnergyHistory" between integrated and gated loudness anlysis
-		int32 NumIntegratedHopHistoryResults = 1;
-
-		TArray<float> ChannelPerceptualEnergy;
-		TArray<float> OverallPerceptualEnergy;
-		TArray<float> InternalWindow;
-
-		// Per channel stats for short-term, integrated and gated loudness, and optional loudness range. 
-		TSortedMap<int32, FAggregateLoudnessStats> AggregateLoudnessStats;
-		// Per channel state for calculating Loudness Range.
-		TSortedMap<int32, FLoudnessRange> LoudnessRangeStatePerChannel;
-
-		// Channel interleaved history buffer of energy measurements. These are 
-		// taken once per an loudness analysis, not necessarily once per a audio window. 
-		TCircularBuffer<float> ChannelEnergyHistory;
-		TCircularBuffer<float> OverallEnergyHistory;
-
-		TSlidingBuffer<float> InternalBuffer;
-		FMultichannelLoudnessAnalyzer Analyzer;
+		FLKFSAnalyzer LKFSAnalyzer;
 	};
 
 	/**
diff --git a/Engine/Plugins/Runtime/AudioSynesthesia/Source/AudioSynesthesiaCore/Public/LKFSNRTFactory.h b/Engine/Plugins/Runtime/AudioSynesthesia/Source/AudioSynesthesiaCore/Public/LKFSNRTFactory.h
--- a/Engine/Plugins/Runtime/AudioSynesthesia/Source/AudioSynesthesiaCore/Public/LKFSNRTFactory.h
+++ b/Engine/Plugins/Runtime/AudioSynesthesia/Source/AudioSynesthesiaCore/Public/LKFSNRTFactory.h
@@ -4,8 +4,8 @@
 
 #include "Containers/ArrayView.h"
 #include "Containers/StridedView.h"
+#include "DSP/LoudnessAnalyzer.h"
 #include "DSP/SlidingWindow.h"
-#include "LoudnessAnalyzer.h"
 #include "IAudioAnalyzerNRTInterface.h"
 
 #include "LKFSNRTFactory.generated.h"
diff --git a/Engine/Plugins/Runtime/AudioSynesthesia/Source/AudioSynesthesiaCore/Public/LoudnessAnalyzer.h b/Engine/Plugins/Runtime/AudioSynesthesia/Source/AudioSynesthesiaCore/Public/LoudnessAnalyzer.h
--- a/Engine/Plugins/Runtime/AudioSynesthesia/Source/AudioSynesthesiaCore/Public/LoudnessAnalyzer.h
+++ b/Engine/Plugins/Runtime/AudioSynesthesia/Source/AudioSynesthesiaCore/Public/LoudnessAnalyzer.h
@@ -1,237 +1,7 @@
 // Copyright Epic Games, Inc. All Rights Reserved.
-
 #pragma once
 
-#include "DSP/AudioFFT.h"
-
-#define UE_API AUDIOSYNESTHESIACORE_API
-
-namespace Audio { class IFFTAlgorithm; }
-
-namespace Audio
-{
-    /** Equal loudness curves from: ANSI Standards S1.4-1983 and S1.42-2001 */
-	enum class ELoudnessCurveType : uint8
-	{
-		A,
-		B,
-		C,
-		D,
-		K,
-		None
-	};
-
-	/** Generate equal loudness weight for magnitude spectrum using curve A */
-    AUDIOSYNESTHESIACORE_API float GetEqualLoudnessAWeigthForFrequency(const float Freq);
-
-	/** Generate equal loudness weight for magnitude spectrum using curve B */
-    AUDIOSYNESTHESIACORE_API float GetEqualLoudnessBWeigthForFrequency(const float Freq);
-
-	/** Generate equal loudness weight for magnitude spectrum using curve C */
-    AUDIOSYNESTHESIACORE_API float GetEqualLoudnessCWeigthForFrequency(const float Freq);
-
-	/** Generate equal loudness weight for magnitude spectrum using curve D */
-    AUDIOSYNESTHESIACORE_API float GetEqualLoudnessDWeigthForFrequency(const float Freq);
-
-	/** Generate equal loudness weight for magnitude spectrum using curve None */
-    AUDIOSYNESTHESIACORE_API float GetEqualLoudnessNoneWeigthForFrequency(const float Freq);
-
-	enum class ELoudnessAnalyzerScalingMethod : uint8
-	{
-		// Normalize by FFTSize and having perceptual curves be unity gain at 1khz. 
-		// 
-		// Note: This method does not correctly account for different window functions
-		// and also reduces the energy by half. It remains here for backwards
-		// compatibility.
-		BackCompatible,
-
-		// Normalize by Window Function. Use this when performing loudness analysis
-		// which must match the LKFS standardized measurements. 
-		Corrected
-	};
-
-	/** Settings for FLoudnessAnalyzer */
-    struct FLoudnessAnalyzerSettings
-    {
-		/** Size of FFT describe in number of samples */
-        int32 FFTSize = 8192;
-
-		/** Type of window to be applied to input audio */
-        EWindowType WindowType = EWindowType::Blackman;
-
-		/** Type of equal loudness curve to weight power spectrum */
-        ELoudnessCurveType LoudnessCurveType = ELoudnessCurveType::D;
-
-		/** Minimum frequency for analyzing audio */
-        float MinAnalysisFrequency = 20.f;
-
-		/** Maximum frequency for analyzing audio */
-        float MaxAnalysisFrequency = 10000.f;
-
-		/** Size of the audio analysis window in samples.
-		 *
-		 * If this is less than 1, it will become equal to the FFTSize.
-		 * Otherwise, this must be less than or equal to the FFT Size.
-		 */
-		int32 WindowSize = -1;
-
-		/** Method to use for scaling and normalizing data. Corrected is preferred
-		 * in terms of correctness. BackCompatible is available for backwards compatibility.
-		 */
-		ELoudnessAnalyzerScalingMethod ScalingMethod = ELoudnessAnalyzerScalingMethod::BackCompatible;
-    };
-
-	/** FLoudnessAnalyzer
-	 *
-	 * FLoudnessAnalyzer will calculate perceptual energy or loudness one window at a time. 
-	 * This analyzer only handles mono audio. See FMultichannelLoudnessAnalyzer for multichannel
-	 * support.
-	 */
-    class FLoudnessAnalyzer
-    {
-        public:
-			/** Construct analyzer */
-            UE_API FLoudnessAnalyzer(float InSampleRate, const FLoudnessAnalyzerSettings& InSettings);
-
-			/**
-			 * Calculate the perceptual energy for the input samples.  The number of 
-			 * input samples must be equal to the FFTSize supplied in FLoudnessAnalyzerSettings.
-			 */
-            UE_API float CalculatePerceptualEnergy(TArrayView<const float> InSampleView);
-
-			/**
-			 * Calculate the loudness for the input samples.  The number of  input samples must 
-			 * be equal to the FFTSize supplied in FLoudnessAnalyzerSettings.
-			 */
-            UE_API float CalculateLoudness(TArrayView<const float> InSampleView);
-
-			/**
-			 * Convert from perceptual energy to loudness [Loudness = 10 * log_10(Energy)]
-			 */
-			static UE_API float ConvertPerceptualEnergyToLoudness(float InPerceptualEnergy);
-			
-			/**
-			 * Convert from perceptual energy to loudness [LKFS = -0.691 + 10 * log_10(Energy)]
-			 */
-			static UE_API float ConvertPerceptualEnergyToLKFS(float InPerceptualEnergy);
-
-			/**
-			 * Return const reference to settings used inside this analyzer.
-			 */
-            UE_API const FLoudnessAnalyzerSettings& GetSettings() const;
-        private:
-
-			float CalculateEnergyScaling() const;
-			float GetEqualLoudnessWeightForFrequency(float InHz) const;
-			float CalculateEnergyScalingBackCompatible() const;
-			float CalculateEnergyScalingCorrected() const;
-			float CalculateWindowEnergyScaling() const;
-
-            FLoudnessAnalyzerSettings Settings;
-
-			/** Scale factor used to normalize between differing window sizes and fft sizes */
-            float EnergyScale = 1.f;
-
-            int32 MinFreqIndex = 0;
-            int32 MaxFreqIndex = 0;
-
-			/** Number of FFTFrequency bins worth looking at. */
-			int32 FFTFreqSize = 0;
-			int32 ActualFFTSize = 0;
-
-			TUniquePtr<IFFTAlgorithm> FFT;
-
-			FAlignedFloatBuffer Window;
-            FAlignedFloatBuffer WindowedSamples;
-            FAlignedFloatBuffer ComplexSpectrum;
-            FAlignedFloatBuffer RealSpectrum;
-
-            FAlignedFloatBuffer CurveWeights;
-    };
-
-
-	/** FMultichannelLoudnessAnalyzer
-	 *
-	 * FMultichannelLoudnessAnalyzer will calculate perceptual energy or loudness one window 
-	 * at a time. This module can handle one or more channels at a time. It expects 
-	 * interleaved audio samples.
-	 */
-    class FMultichannelLoudnessAnalyzer 
-    {
-        public:
-			/** Construct analyzer. */
-            UE_API FMultichannelLoudnessAnalyzer(float InSampleRate, const FLoudnessAnalyzerSettings& InSettings);
-
-			/** 
-			 * Calculate perceptual energy. 
-			 *
-			 * InSampleView		  InSampleView is expected to contain interleaved audio with 
-			 *                    (InSettings.FFTSize * InNumChannels) number of samples.
-			 *
-			 * InNumChannels      Number of input channels.
-			 *
-			 * OutChannelEnergies Will be populated with energies on a per channel basis.
-			 *
-			 * The overall perceptual energy is returned.
-			 */
-			UE_DEPRECATED(5.7, "Use method which accepts a TArrayView<float> for OutChannelEnergies")
-            UE_API float CalculatePerceptualEnergy(TArrayView<const float> InSampleView, const int32 InNumChannels, TArray<float>& OutChannelEnergies); 
-
-			/** 
-			 * Calculate perceptual energy. 
-			 *
-			 * InSampleView		  InSampleView is expected to contain interleaved audio with 
-			 *                    (InSettings.FFTSize * InNumChannels) number of samples.
-			 *
-			 * InNumChannels      Number of input channels.
-			 *
-			 * OutChannelEnergies Will be populated with energies on a per channel basis.
-			 *
-			 * The overall perceptual energy is returned.
-			 */
-            UE_API float CalculatePerceptualEnergy(TArrayView<const float> InSampleView, const int32 InNumChannels, TArrayView<float> OutChannelEnergies); 
-
-			/** 
-			 * Calculate loudness
-			 *
-			 * InSampleView		  InSampleView is expected to contain interleaved audio with 
-			 *                    (InSettings.FFTSize * InNumChannels) number of samples.
-			 *
-			 * InNumChannels      Number of input channels.
-			 *
-			 * OutChannelLoudness Will be populated with energies on a per channel basis.
-			 *
-			 * The overall loudness is returned.
-			 */
-			UE_DEPRECATED(5.7, "Use method which accepts a TArrayView<float> for OutChannelLoudness")
-            UE_API float CalculateLoudness(TArrayView<const float> InSampleView, const int32 InNumChannels, TArray<float>& OutChannelLoudness);
-
-			/** 
-			 * Calculate loudness
-			 *
-			 * InSampleView		  InSampleView is expected to contain interleaved audio with 
-			 *                    (InSettings.FFTSize * InNumChannels) number of samples.
-			 *
-			 * InNumChannels      Number of input channels.
-			 *
-			 * OutChannelLoudness Will be populated with energies on a per channel basis.
-			 *
-			 * The overall loudness is returned.
-			 */
-            UE_API float CalculateLoudness(TArrayView<const float> InSampleView, const int32 InNumChannels, TArrayView<float> OutChannelLoudness);
-
-			/**
-			 * Return const reference to settings used inside this analyzer.
-			 */
-            UE_API const FLoudnessAnalyzerSettings& GetSettings() const;
-
-        private:
-			void InitializeChannelWeights(int32 InNumChannels);
-
-            FAlignedFloatBuffer MonoBuffer;
-            TArray<float> ChannelWeights;
-            FLoudnessAnalyzer Analyzer;
-    };
-}
+#include "Misc/CoreMiscDefines.h"
+UE_DEPRECATED_HEADER(5.8, "Include DSP/LoudnessAnalyzer.h from the SignalProcessing module instead")
 
-#undef UE_API
+#include "DSP/LoudnessAnalyzer.h"
diff --git a/Engine/Plugins/Runtime/AudioSynesthesia/Source/AudioSynesthesiaCore/Public/LoudnessFactory.h b/Engine/Plugins/Runtime/AudioSynesthesia/Source/AudioSynesthesiaCore/Public/LoudnessFactory.h
--- a/Engine/Plugins/Runtime/AudioSynesthesia/Source/AudioSynesthesiaCore/Public/LoudnessFactory.h
+++ b/Engine/Plugins/Runtime/AudioSynesthesia/Source/AudioSynesthesiaCore/Public/LoudnessFactory.h
@@ -2,7 +2,7 @@
 
 #pragma once
 
-#include "LoudnessAnalyzer.h"
+#include "DSP/LoudnessAnalyzer.h"
 #include "IAudioAnalyzerInterface.h"
 
 namespace Audio { template <typename InSampleType> class TSlidingBuffer; }
diff --git a/Engine/Plugins/Runtime/AudioSynesthesia/Source/AudioSynesthesiaCore/Public/LoudnessNRTFactory.h b/Engine/Plugins/Runtime/AudioSynesthesia/Source/AudioSynesthesiaCore/Public/LoudnessNRTFactory.h
--- a/Engine/Plugins/Runtime/AudioSynesthesia/Source/AudioSynesthesiaCore/Public/LoudnessNRTFactory.h
+++ b/Engine/Plugins/Runtime/AudioSynesthesia/Source/AudioSynesthesiaCore/Public/LoudnessNRTFactory.h
@@ -2,7 +2,7 @@
 
 #pragma once
 
-#include "LoudnessAnalyzer.h"
+#include "DSP/LoudnessAnalyzer.h"
 #include "IAudioAnalyzerNRTInterface.h"
 
 namespace Audio { template <typename InSampleType> class TSlidingBuffer; }
diff --git a/Engine/Source/Runtime/AudioMixerCore/Private/AudioMixer.cpp b/Engine/Source/Runtime/AudioMixerCore/Private/AudioMixer.cpp
--- a/Engine/Source/Runtime/AudioMixerCore/Private/AudioMixer.cpp
+++ b/Engine/Source/Runtime/AudioMixerCore/Private/AudioMixer.cpp
@@ -4,6 +4,7 @@
 #include "Async/Async.h"
 #include "AudioDefines.h"
 #include "AudioMixerTrace.h"
+#include "DSP/AudioChannelUtils.h"
 #include "DSP/BufferVectorOperations.h"
 #include "DSP/FloatArrayMath.h"
 #include "HAL/RunnableThread.h"
@@ -825,99 +826,9 @@ namespace Audio
 		return ReturnVal;
 	}
 
-	/** The default channel orderings to use when using pro audio interfaces while still supporting surround sound. */
-	static EAudioMixerChannel::Type DefaultChannelOrder[AUDIO_MIXER_MAX_OUTPUT_CHANNELS];
-
-	static void InitializeDefaultChannelOrder()
-	{
-		static bool bInitialized = false;
-		if (bInitialized)
-		{
-			return;
-		}
-
-		bInitialized = true;
-
-		// Create a hard-coded default channel order
-		check(UE_ARRAY_COUNT(DefaultChannelOrder) == AUDIO_MIXER_MAX_OUTPUT_CHANNELS);
-		DefaultChannelOrder[0] = EAudioMixerChannel::FrontLeft;
-		DefaultChannelOrder[1] = EAudioMixerChannel::FrontRight;
-		DefaultChannelOrder[2] = EAudioMixerChannel::FrontCenter;
-		DefaultChannelOrder[3] = EAudioMixerChannel::LowFrequency;
-		DefaultChannelOrder[4] = EAudioMixerChannel::SideLeft;
-		DefaultChannelOrder[5] = EAudioMixerChannel::SideRight;
-		DefaultChannelOrder[6] = EAudioMixerChannel::BackLeft;
-		DefaultChannelOrder[7] = EAudioMixerChannel::BackRight;
-
-		bool bOverridden = false;
-		EAudioMixerChannel::Type ChannelMapOverride[AUDIO_MIXER_MAX_OUTPUT_CHANNELS];
-		for (int32 i = 0; i < AUDIO_MIXER_MAX_OUTPUT_CHANNELS; ++i)
-		{
-			ChannelMapOverride[i] = DefaultChannelOrder[i];
-		}
-
-		// Now check the ini file to see if this is overridden
-		for (int32 i = 0; i < AUDIO_MIXER_MAX_OUTPUT_CHANNELS; ++i)
-		{
-			int32 ChannelPositionOverride = 0;
-
-			const TCHAR* ChannelName = EAudioMixerChannel::ToString(DefaultChannelOrder[i]);
-			if (GConfig->GetInt(TEXT("AudioDefaultChannelOrder"), ChannelName, ChannelPositionOverride, GEngineIni))
-			{
-				if (ChannelPositionOverride >= 0 && ChannelPositionOverride < AUDIO_MIXER_MAX_OUTPUT_CHANNELS)
-				{
-					bOverridden = true;
-					ChannelMapOverride[ChannelPositionOverride] = DefaultChannelOrder[i];
-				}
-				else
-				{
-					UE_LOG(LogAudioMixer, Error, TEXT("Invalid channel index '%d' in AudioDefaultChannelOrder in ini file."), i);
-					bOverridden = false;
-					break;
-				}
-			}
-		}
-
-		// Now validate that there's no duplicates.
-		if (bOverridden)
-		{
-			bool bIsValid = true;
-			for (int32 i = 0; i < AUDIO_MIXER_MAX_OUTPUT_CHANNELS; ++i)
-			{
-				for (int32 j = 0; j < AUDIO_MIXER_MAX_OUTPUT_CHANNELS; ++j)
-				{
-					if (j != i && ChannelMapOverride[j] == ChannelMapOverride[i])
-					{
-						bIsValid = false;
-						break;
-					}
-				}
-			}
-
-			if (!bIsValid)
-			{
-				UE_LOG(LogAudioMixer, Error, TEXT("Invalid channel index or duplicate entries in AudioDefaultChannelOrder in ini file."));
-			}
-			else
-			{
-				for (int32 i = 0; i < AUDIO_MIXER_MAX_OUTPUT_CHANNELS; ++i)
-				{
-					DefaultChannelOrder[i] = ChannelMapOverride[i];
-				}
-			}
-		}
-	}
-
 	bool IAudioMixerPlatformInterface::GetChannelTypeAtIndex(const int32 Index, EAudioMixerChannel::Type& OutType)
 	{
-		InitializeDefaultChannelOrder();
-
-		if (Index >= 0 && Index < AUDIO_MIXER_MAX_OUTPUT_CHANNELS)
-		{
-			OutType = DefaultChannelOrder[Index];
-			return true;
-		}
-		return false;
+		return ChannelUtils::GetChannelTypeAtIndex(Index, OutType);
 	}
 
 	bool IAudioMixer::ShouldUseThreadedDeviceSwap()
diff --git a/Engine/Source/Runtime/SignalProcessing/Private/AudioChannelUtils.cpp b/Engine/Source/Runtime/SignalProcessing/Private/AudioChannelUtils.cpp
new file mode 100644
--- /dev/null
+++ b/Engine/Source/Runtime/SignalProcessing/Private/AudioChannelUtils.cpp
@@ -0,0 +1,135 @@
+// Copyright Epic Games, Inc. All Rights Reserved.
+
+#include "DSP/AudioChannelUtils.h"
+
+#include "AudioMixer.h"
+#include "Misc/ConfigCacheIni.h"
+
+namespace Audio::ChannelUtils
+{	
+	bool GetChannelTypeAtIndex(const int32 Index, EAudioMixerChannel::Type& OutType)
+	{
+		/** The default channel orderings to use when using pro audio interfaces while still supporting surround sound. */
+		auto CreateDefaultChannelOrder = []()
+		{
+			// Create a hard-coded default channel order
+			TArray<EAudioMixerChannel::Type> DefaultChannelOrder
+			{
+				EAudioMixerChannel::FrontLeft,
+				EAudioMixerChannel::FrontRight,
+				EAudioMixerChannel::FrontCenter,
+				EAudioMixerChannel::LowFrequency,
+				EAudioMixerChannel::SideLeft,
+				EAudioMixerChannel::SideRight,
+				EAudioMixerChannel::BackLeft,
+				EAudioMixerChannel::BackRight
+			};
+
+			TArray<EAudioMixerChannel::Type> ChannelMapOverride = DefaultChannelOrder;
+
+			check(DefaultChannelOrder.Num() == AUDIO_MIXER_MAX_OUTPUT_CHANNELS);
+			check(ChannelMapOverride.Num()  == AUDIO_MIXER_MAX_OUTPUT_CHANNELS);
+
+			// Now check the ini file to see if this is overridden
+			bool bOverridden = false;
+
+			for (int32 i = 0; i < AUDIO_MIXER_MAX_OUTPUT_CHANNELS; ++i)
+			{
+				int32 ChannelPositionOverride = 0;
+
+				const TCHAR* ChannelName = EAudioMixerChannel::ToString(DefaultChannelOrder[i]);
+
+				if (GConfig->GetInt(TEXT("AudioDefaultChannelOrder"), ChannelName, ChannelPositionOverride, GEngineIni))
+				{
+					if (ChannelPositionOverride >= 0 && ChannelPositionOverride < AUDIO_MIXER_MAX_OUTPUT_CHANNELS)
+					{
+						bOverridden = true;
+						ChannelMapOverride[ChannelPositionOverride] = DefaultChannelOrder[i];
+					}
+					else
+					{
+						UE_LOG(LogSignalProcessing, Error, TEXT("Invalid channel index '%d' in AudioDefaultChannelOrder in ini file."), i);
+						bOverridden = false;
+						break;
+					}
+				}
+			}
+
+			// Now validate that there's no duplicates.
+			if (bOverridden)
+			{
+				bool bIsValid = true;
+
+				for (int32 i = 0; i < AUDIO_MIXER_MAX_OUTPUT_CHANNELS; ++i)
+				{
+					for (int32 j = 0; j < AUDIO_MIXER_MAX_OUTPUT_CHANNELS; ++j)
+					{
+						if (j != i && ChannelMapOverride[j] == ChannelMapOverride[i])
+						{
+							bIsValid = false;
+							break;
+						}
+					}
+				}
+
+				if (!bIsValid)
+				{
+					UE_LOG(LogSignalProcessing, Error, TEXT("Invalid channel index or duplicate entries in AudioDefaultChannelOrder in ini file."));
+				}
+				else
+				{
+					for (int32 i = 0; i < AUDIO_MIXER_MAX_OUTPUT_CHANNELS; ++i)
+					{
+						DefaultChannelOrder[i] = ChannelMapOverride[i];
+					}
+				}
+			}
+
+			return DefaultChannelOrder;
+		};
+
+		static const TArray<EAudioMixerChannel::Type> DefaultChannelOrder = CreateDefaultChannelOrder();
+
+		if (Index >= 0 && Index < AUDIO_MIXER_MAX_OUTPUT_CHANNELS)
+		{
+			OutType = DefaultChannelOrder[Index];
+			return true;
+		}
+
+		return false;
+	}
+	
+	float GetChannelWeight(EAudioMixerChannel::Type InType) 
+	{
+		static_assert(static_cast<int32>(EAudioMixerChannel::ChannelTypeCount) == 19, "Possibly missing channel type");
+
+		switch (InType)
+		{
+			case EAudioMixerChannel::LowFrequency:
+				return 0.f; // LFE channels ignored
+				
+			case EAudioMixerChannel::BackLeft:
+			case EAudioMixerChannel::BackRight:
+			case EAudioMixerChannel::SideLeft:
+			case EAudioMixerChannel::SideRight:
+			// If elevation angle is less than 30 degrees and azimuth is between 60 and 120, set channel weight to sqrt(2)
+				return FMath::Sqrt(2.f);
+
+			case EAudioMixerChannel::FrontLeft:
+			case EAudioMixerChannel::FrontRight:
+			case EAudioMixerChannel::FrontCenter:
+			case EAudioMixerChannel::FrontLeftOfCenter:
+			case EAudioMixerChannel::FrontRightOfCenter:
+			case EAudioMixerChannel::BackCenter:
+			case EAudioMixerChannel::TopCenter:
+			case EAudioMixerChannel::TopFrontLeft:
+			case EAudioMixerChannel::TopFrontCenter:
+			case EAudioMixerChannel::TopFrontRight:
+			case EAudioMixerChannel::TopBackLeft:
+			case EAudioMixerChannel::TopBackCenter:
+			case EAudioMixerChannel::TopBackRight:
+			default:
+				return 1.f;
+		}
+	}
+} // namespace Audio::ChannelUtils
diff --git a/Engine/Source/Runtime/SignalProcessing/Private/LKFSAnalyzer.cpp b/Engine/Source/Runtime/SignalProcessing/Private/LKFSAnalyzer.cpp
new file mode 100644
--- /dev/null
+++ b/Engine/Source/Runtime/SignalProcessing/Private/LKFSAnalyzer.cpp
@@ -0,0 +1,493 @@
+// Copyright Epic Games, Inc. All Rights Reserved.
+
+#include "DSP/LKFSAnalyzer.h"
+
+#include "DSP/LKFSUtils.h"
+#include "DSP/SlidingWindow.h"
+
+namespace Audio
+{
+	namespace LKFSAnalyzerPrivate
+	{
+		float LKFSToEnergy(float InLKFS)
+		{
+			return FMath::Pow(10, (InLKFS + 0.691f) / 10.f);
+		}
+
+		// Calculates average value from a circular buffer. Allows for interleaved data by providing a "stride". 
+		float CalculateAverageFromHistory(const TCircularBuffer<float>& InHistory, int32 InStartIndex, int32 InEndIndex, int32 InStride)
+		{
+			const int32 Num = (InEndIndex - InStartIndex) / InStride;
+			if (Num <= 0)
+			{
+				return 0.0f;
+			}
+
+			float Accum = 0.0f;
+			for (int32 Index = InStartIndex; Index < InEndIndex; Index += InStride)
+			{
+				Accum += InHistory[Index];
+			}
+
+			return Accum /= Num;
+		}
+
+		// Calculates gated average value from a circular buffer. Allows for interleaved data by providing a "stride". 
+		float CalculateGatedAverageFromHistory(const TCircularBuffer<float>& InHistory, int32 InStartIndex, int32 InEndIndex, int32 InStride, float InGate)
+		{
+			float Accum = 0.0f;
+			int32 Num = 0;
+
+			for (int32 Index = InStartIndex; Index < InEndIndex; Index += InStride)
+			{
+				const float Val = InHistory[Index];
+
+				if (Val > InGate)
+				{
+					Accum += Val;
+					++Num;
+				}
+			}
+
+			return (Num > 0) ? Accum /= Num : InGate;
+		}
+
+		float CalculateGatedLoudnessFromHistory(const TCircularBuffer<float>& InHistory, int32 InStartIndex, int32 InEndIndex, int32 InStride)
+		{
+			// Gated Loudness only examines loudnesses above a certain gating loudness. 
+			//
+			// It is a 2 pass algorithm 
+			// Pass 1: Use a set loudness threshold of -70lkfs to determine an initial gated loudness
+			// Pass 2: Use the initial gated loudness to derive a new loudness threshold and calculate the gate loudness again. 
+			constexpr float InitLoudnessThreshold = -70.0f;
+
+			float InitEnergyThreshold = LKFSToEnergy(InitLoudnessThreshold);
+			const float GatedEnergyPass1 = CalculateGatedAverageFromHistory(InHistory, InStartIndex, InEndIndex, InStride, InitEnergyThreshold);
+
+			const float NewThreshold = GatedEnergyPass1 * 0.1f;
+			const float GatedEnergyPass2 = CalculateGatedAverageFromHistory(InHistory, InStartIndex, InEndIndex, InStride, NewThreshold);
+
+			return FLoudnessAnalyzer::ConvertPerceptualEnergyToLKFS(GatedEnergyPass2);
+		}
+
+		// Implementation copied from FFineGrainedPerformanceTracker::GetPercentileValue(...)
+		float GetPercentileValue(TArray<float>& Samples, int32 Percentile)
+		{
+			ensure(Percentile >= 0 && Percentile <= 100);
+
+			int32 Left = 0;
+			int32 Right = Samples.Num() - 1;
+
+			if (Right < 0)
+			{
+				return -1.0f;
+			}
+
+			const int32 PercentileClamped = FMath::Clamp(Percentile, 0, 100);
+			const int32 PercentileOrdinal = FMath::RoundToInt32((float(PercentileClamped) * Right) / 100);
+
+			// this is quickselect (see http://en.wikipedia.org/wiki/Quickselect for details).
+			while (Right != Left)
+			{
+				// partition
+				int32 MovingLeft = Left - 1;
+				int32 MovingRight = Right;
+				float Pivot = Samples[MovingRight];
+				for (; ; )
+				{
+					while (Samples[++MovingLeft] < Pivot);
+					while (Samples[--MovingRight] > Pivot)
+					{
+						if (MovingRight == Left)
+						{
+							break;
+						}
+					}
+
+					if (MovingLeft >= MovingRight)
+					{
+						break;
+					}
+
+					const float Temp = Samples[MovingLeft];
+					Samples[MovingLeft] = Samples[MovingRight];
+					Samples[MovingRight] = Temp;
+				}
+
+				const float Temp = Samples[MovingLeft];
+				Samples[MovingLeft] = Samples[Right];
+				Samples[Right] = Temp;
+
+				// now we're pivoted around MovingLeft
+				// decide what part K-th largest belong to
+				if (MovingLeft > PercentileOrdinal)
+				{
+					Right = MovingLeft - 1;
+				}
+				else if (MovingLeft < PercentileOrdinal)
+				{
+					Left = MovingLeft + 1;
+				}
+				else
+				{
+					// we hit exactly the value we need, no need to sort further
+					break;
+				}
+			}
+
+			return Samples[PercentileOrdinal];
+		}
+	} // namespace LKFSAnalyzerPrivate
+
+
+	//----------------------------------------------------------------------------------------
+	// FLKFSAnalyzer
+	//----------------------------------------------------------------------------------------
+	FLKFSAnalyzer::FLKFSAnalyzer(int32 InNumChannels, float InSampleRate, const FLKFSAnalyzerSettings& InSettings)
+		: Settings(InSettings)
+		, NumChannels(InNumChannels)
+		, SampleRate(InSampleRate)
+		, ChannelEnergyHistory(1)
+		, OverallEnergyHistory(1)
+		, LoudnessAnalyzer(InSampleRate, InSettings.LoudnessAnalyzerSettings)
+	{
+		check(NumChannels > 0);
+		check(SampleRate > 0.f);
+		check(InSettings.ShortTermLoudnessDuration > 0.f);
+		check(InSettings.IntegratedLoudnessDuration > 0.f);
+		check(InSettings.IntegratedLoudnessAnalysisPeriod> 0.f);
+
+		int32 NumSlidingWindowFrames = LoudnessAnalyzer.GetSettings().WindowSize;
+
+		int32 NumAnalysisFrames = FMath::CeilToInt(InSettings.AnalysisWindowDuration * SampleRate);
+		int32 NumAnalysisHopFrames = FMath::CeilToInt(InSettings.AnalysisPeriod * SampleRate);
+
+		// There may be multiple calls to the LoudnessAnalyzer to produce a single 
+		// FLoudnessDatum. The sliding window hop is tuned here so that it best 
+		// matches the desired AnalysisPeriod while maintaining between a 25% to 
+		// 75% window overlap.
+		LKFSUtils::TuneSlidingWindowHopSize(NumAnalysisHopFrames, NumSlidingWindowFrames, NumSlidingWindowHopFrames, NumAnalysisHopWindows);
+
+		// Determine how many windows to analyze per loudness result
+		NumAnalysisWindows = FMath::Max(1, ((NumAnalysisFrames - (NumSlidingWindowFrames / 2)) / NumSlidingWindowHopFrames) + 1);
+
+		// Include NumChannels in calculating window size because windows are generated 
+		// with interleaved samples and deinterleaved later. 
+		int32 NumSlidingWindowSamples = NumSlidingWindowFrames * NumChannels;
+		int32 NumSlidingWindowHopSamples = NumSlidingWindowHopFrames * NumChannels;
+
+		InternalBuffer = TSlidingBuffer<float>(NumSlidingWindowSamples, NumSlidingWindowHopSamples);
+		InternalWindow.AddUninitialized(NumSlidingWindowSamples);
+		ChannelPerceptualEnergy.AddZeroed(NumChannels * NumAnalysisWindows);
+		OverallPerceptualEnergy.AddZeroed(NumAnalysisWindows);
+
+		// Initialize settings for calculating aggregate stats
+		const int32 NumResultHopFrames = NumAnalysisHopWindows * NumSlidingWindowHopFrames;
+		NumShortTermHistoryResults     = FMath::Max(1, FMath::DivideAndRoundUp(static_cast<int32>(InSettings.ShortTermLoudnessDuration * SampleRate), NumResultHopFrames));
+		NumIntegratedHistoryResults    = FMath::Max(1, FMath::DivideAndRoundUp(static_cast<int32>(InSettings.IntegratedLoudnessDuration * SampleRate), NumResultHopFrames));
+		NumIntegratedHopHistoryResults = FMath::Max(1, FMath::DivideAndRoundUp(static_cast<int32>(InSettings.IntegratedLoudnessAnalysisPeriod * SampleRate), NumResultHopFrames));
+
+		const int32 MaxHistory = FMath::Max(NumIntegratedHistoryResults, NumShortTermHistoryResults);
+		ChannelEnergyHistory = TCircularBuffer<float>(MaxHistory * NumChannels);
+		OverallEnergyHistory = TCircularBuffer<float>(MaxHistory);
+
+		for (int32 ChannelIdx = 0; ChannelIdx < NumChannels; ChannelIdx++)
+		{
+			AggregateLoudnessStats.Add(ChannelIdx);
+		}
+
+		AggregateLoudnessStats.Add(FLKFSAnalyzerResults::ChannelIndexOverall);
+
+		if (InSettings.bCalculatePerChannelLoudnessRange)
+		{
+			for (int32 ChannelIdx = 0; ChannelIdx < NumChannels; ChannelIdx++)
+			{
+				LoudnessRangeStatePerChannel.Emplace(ChannelIdx, static_cast<uint32>(NumIntegratedHistoryResults));
+			}
+		}
+
+		if (InSettings.bCalculateOverallLoudnessRange)
+		{
+			LoudnessRangeStatePerChannel.Emplace(FLKFSAnalyzerResults::ChannelIndexOverall, static_cast<uint32>(NumIntegratedHistoryResults));
+		}
+	}
+
+	const FLKFSAnalyzerSettings& FLKFSAnalyzer::GetSettings() const
+	{
+		return Settings;
+	}
+
+	void FLKFSAnalyzer::Analyze(TArrayView<const float> InAudio, FLKFSAnalyzerResults& OutLKFSResults)
+	{
+		TAutoSlidingWindow<float> SlidingWindow(InternalBuffer, InAudio, InternalWindow);
+
+		AnalyzeSlidingWindow(SlidingWindow, OutLKFSResults);
+	}
+
+	void FLKFSAnalyzer::AnalyzeSlidingWindow(TAutoSlidingWindow<float>& InSlidingWindow, FLKFSAnalyzerResults& OutResult)
+	{
+		check(NumAnalysisWindows > 0);
+		check(NumAnalysisHopWindows > 0);
+
+		for (const TArray<float>& Window : InSlidingWindow)
+		{
+			AnalyzeWindow(Window);
+
+			if ((NumAnalyzedBuffers >= NumAnalysisWindows) && ((NumAnalyzedBuffers - NumAnalysisWindows) % NumAnalysisHopWindows) == 0)
+			{
+				AddCurrentLoudnessDataToResult(OutResult);
+			}
+		}
+	}
+
+	void FLKFSAnalyzer::AnalyzeWindow(TArrayView<const float> InWindow)
+	{
+		const int32 WrappedWindowIdx = (NumAnalyzedBuffers % NumAnalysisWindows);
+		const int32 ChannelPerceptualEnergyWritePos = WrappedWindowIdx * NumChannels;
+
+		// Calculate perceptual energy
+		OverallPerceptualEnergy[WrappedWindowIdx] = LoudnessAnalyzer.CalculatePerceptualEnergy(InWindow, NumChannels, MakeArrayView<float>(ChannelPerceptualEnergy.GetData() + ChannelPerceptualEnergyWritePos, NumChannels));
+
+		// Update counters
+		NumAnalyzedBuffers++;
+	}
+
+	void FLKFSAnalyzer::AddCurrentLoudnessDataToResult(FLKFSAnalyzerResults& OutResult)
+	{
+		AddCurrentLoudnessDataToHistory();
+		UpdateAggregateLoudnessStats();
+
+		// Accounting for all the frames positions analyzed in creating the data, determine the center frame. 
+		const int32 CenterFramePos = ((2 * NumAnalyzedBuffers - NumAnalysisWindows) * NumSlidingWindowHopFrames + LoudnessAnalyzer.GetSettings().WindowSize) / 2;
+		const float Timestamp = CenterFramePos / SampleRate;
+
+		// Add per channel results
+		for (int32 ChannelIdx = 0; ChannelIdx < NumChannels; ++ChannelIdx)
+		{
+			const int32 HistoryIndex = ((NumEnergyHistory - 1) * NumChannels) + ChannelIdx;
+			const float Energy = ChannelEnergyHistory[HistoryIndex];
+			const FAggregateLoudnessStats& Stats = AggregateLoudnessStats[ChannelIdx];
+
+			OutResult.Add(
+				FLKFSResult
+				{
+					.Channel = ChannelIdx, 
+					.Timestamp = Timestamp, 
+					.Energy = Energy, 
+					.Loudness = FLoudnessAnalyzer::ConvertPerceptualEnergyToLKFS(Energy),
+					.ShortTermLoudness = Stats.ShortTermLoudness, 
+					.IntegratedLoudness = Stats.IntegratedLoudness,
+					.GatedLoudness = Stats.GatedLoudness,
+					.LoudnessRange = Stats.LoudnessRange,
+				}
+			);
+		}
+
+		// Add overall results
+		{
+			const int32 HistoryIndex = (NumEnergyHistory - 1);
+			const float Energy = OverallEnergyHistory[HistoryIndex];
+			const FAggregateLoudnessStats& Stats = AggregateLoudnessStats[FLKFSAnalyzerResults::ChannelIndexOverall];
+
+			OutResult.Add(
+				FLKFSResult
+				{
+					.Channel = FLKFSAnalyzerResults::ChannelIndexOverall,
+					.Timestamp = Timestamp, 
+					.Energy = Energy, 
+					.Loudness = FLoudnessAnalyzer::ConvertPerceptualEnergyToLKFS(Energy),
+					.ShortTermLoudness = Stats.ShortTermLoudness, 
+					.IntegratedLoudness = Stats.IntegratedLoudness,
+					.GatedLoudness = Stats.GatedLoudness,
+					.LoudnessRange = Stats.LoudnessRange,
+				}
+			);
+		}
+	}
+
+	void FLKFSAnalyzer::AddCurrentLoudnessDataToHistory()
+	{
+		// Compute loudness per channel and store history
+		for (int32 ChannelIdx = 0; ChannelIdx < NumChannels; ChannelIdx++)
+		{
+			// Determine average perceptual energy over analysis window
+			float PerceptualEnergy = 0.f;
+			for (int32 Pos = ChannelIdx; Pos < ChannelPerceptualEnergy.Num(); Pos += NumChannels)
+			{
+				PerceptualEnergy += ChannelPerceptualEnergy[Pos];
+			}
+			PerceptualEnergy /= (float) NumAnalysisWindows;
+
+			// Store perceptual energy in history 
+			ChannelEnergyHistory[(NumChannels * NumEnergyHistory) + ChannelIdx] = PerceptualEnergy;
+		}
+
+		// Compute loudness overall
+		{
+			float PerceptualEnergy = 0.f;
+			for (int32 Pos = 0; Pos < OverallPerceptualEnergy.Num(); Pos++)
+			{
+				PerceptualEnergy += OverallPerceptualEnergy[Pos];
+			}
+			PerceptualEnergy /= (float) NumAnalysisWindows;
+
+			OverallEnergyHistory[NumEnergyHistory] = PerceptualEnergy;
+		}
+
+		NumEnergyHistory++;
+	}
+
+	void FLKFSAnalyzer::UpdateAggregateLoudnessStats()
+	{
+		// Short term loudness is always updated, but integrated & gated loudness are 
+		// updated at a lower rate. 
+		bool bUpdateIntegratedLoudness = (NumEnergyHistory == 1) || (NumEnergyHistory % NumIntegratedHopHistoryResults) == 0;
+
+
+		// Per channel results
+		for (int32 ChannelIdx = 0; ChannelIdx < NumChannels; ChannelIdx++)
+		{
+			FAggregateLoudnessStats& Stats = AggregateLoudnessStats[ChannelIdx];
+			FLoudnessRange* LoudnessRange = LoudnessRangeStatePerChannel.Find(ChannelIdx);
+			UpdateAggregateLoudnessStats(ChannelEnergyHistory, ChannelIdx, NumChannels, bUpdateIntegratedLoudness, Stats, LoudnessRange);
+		}
+
+		// Overall results
+		{
+			FAggregateLoudnessStats& Stats = AggregateLoudnessStats[FLKFSAnalyzerResults::ChannelIndexOverall];
+			FLoudnessRange* LoudnessRange = LoudnessRangeStatePerChannel.Find(FLKFSAnalyzerResults::ChannelIndexOverall);
+			UpdateAggregateLoudnessStats(OverallEnergyHistory, 0 /*offset*/, 1/*stride*/, bUpdateIntegratedLoudness, Stats, LoudnessRange);
+		}
+	}
+
+	void FLKFSAnalyzer::UpdateAggregateLoudnessStats(const TCircularBuffer<float>& InHistory, int32 InOffset, int32 InStride, bool bInUpdateIntegratedLoudness, FAggregateLoudnessStats& OutStats, FLoudnessRange* LoudnessRange) const
+	{
+		const int32 EndIndex = (NumEnergyHistory * InStride) + InOffset;
+
+		// Update short term loudness always.
+		const int32 ShortTermLoudnessStartIndex = InOffset + FMath::Max(0, (NumEnergyHistory - NumShortTermHistoryResults) * InStride);
+		const float ShortTermAverageEnergy = LKFSAnalyzerPrivate::CalculateAverageFromHistory(InHistory, ShortTermLoudnessStartIndex, EndIndex, InStride);
+		OutStats.ShortTermLoudness = FLoudnessAnalyzer::ConvertPerceptualEnergyToLKFS(ShortTermAverageEnergy);
+
+		// Only update long term loudness if flag is set. 
+		if (bInUpdateIntegratedLoudness)
+		{
+			const int32 LongTermLoudnessStartIndex = InOffset + FMath::Max(0, (NumEnergyHistory - NumIntegratedHistoryResults * InStride));
+
+			float IntegratedEnergy = LKFSAnalyzerPrivate::CalculateAverageFromHistory(InHistory, LongTermLoudnessStartIndex, EndIndex, InStride);
+			OutStats.IntegratedLoudness = FLoudnessAnalyzer::ConvertPerceptualEnergyToLKFS(IntegratedEnergy);
+			OutStats.GatedLoudness = LKFSAnalyzerPrivate::CalculateGatedLoudnessFromHistory(InHistory, LongTermLoudnessStartIndex, EndIndex, InStride);
+		}
+
+		// Update optional loudness range (if we have LoudnessRange state):
+		if (LoudnessRange)
+		{
+			LoudnessRange->AddShortTermEnergyMeasurement(ShortTermAverageEnergy);
+
+			if (bInUpdateIntegratedLoudness)
+			{
+				OutStats.LoudnessRange = LoudnessRange->CalculateLoudnessRange();
+			}
+		}
+	}
+
+	
+	//----------------------------------------------------------------------------------------
+	// FLKFSAnalyzerResults
+	//----------------------------------------------------------------------------------------
+	int32 FLKFSAnalyzerResults::GetNumChannels() const
+	{
+		const int32 Num = ChannelResults.Num();
+
+		if (Num > 0)
+		{
+			return Num - 1; // Subtract one for storage of "Overall" channel
+		}
+
+		return 0;
+	}
+
+	void FLKFSAnalyzerResults::Add(FLKFSResult InDatum)
+	{
+		TArray<FLKFSResult>& ChannelData = ChannelResults.FindOrAdd(InDatum.Channel);
+		ChannelData.Add(MoveTemp(InDatum));
+	}
+
+	const TArray<FLKFSResult>& FLKFSAnalyzerResults::GetChannelLoudnessResults(int32 ChannelIdx) const
+	{
+		static const TArray<FLKFSResult> EmptyArray;
+
+		const TArray<FLKFSResult>* ChannelData = ChannelResults.Find(ChannelIdx);
+
+		return ChannelData ? *ChannelData : EmptyArray;
+	}
+
+	const TArray<FLKFSResult>& FLKFSAnalyzerResults::GetLoudnessResults() const
+	{
+		return GetChannelLoudnessResults(ChannelIndexOverall);
+	}
+
+
+	//----------------------------------------------------------------------------------------
+	// FLoudnessRange
+	//----------------------------------------------------------------------------------------
+	void FLKFSAnalyzer::FLoudnessRange::AddShortTermEnergyMeasurement(float InShortTermEnergy)
+	{
+		ShortTermEnergyHistory[ShortTermWriteIndex++] = InShortTermEnergy;
+		ShortTermHistorySize = FMath::Min<uint32>(ShortTermHistorySize + 1, MaxShortTermHistorySize);
+	}
+
+	TOptional<FFloatInterval> FLKFSAnalyzer::FLoudnessRange::CalculateLoudnessRange() const
+	{
+		// Based on EBU Tech 3342
+		if (ShortTermHistorySize == 0)
+		{
+			return NullOpt;
+		}
+
+		constexpr float AbsoluteThresholdEnergy = 1.17246636e-7f; // -70 LUFS
+		constexpr float RelativeThresholdEnergy = 0.01f;		  // -20 LU 
+
+		const uint32 LoudnessRangeStartIndex = ShortTermWriteIndex - ShortTermHistorySize;
+		const uint32 LoudnessRangeEndIndex = ShortTermWriteIndex;
+
+		// Apply the absolute-threshold gating
+		const float RefEnergy = LKFSAnalyzerPrivate::CalculateGatedAverageFromHistory(
+			ShortTermEnergyHistory,
+			static_cast<int32>(LoudnessRangeStartIndex),
+			static_cast<int32>(LoudnessRangeEndIndex),
+			1 /*Stride*/,
+			AbsoluteThresholdEnergy
+		);
+
+		// Apply the relative-threshold gating
+		const float RelGateEnergy = FMath::Max(RefEnergy * RelativeThresholdEnergy, AbsoluteThresholdEnergy);
+
+		TArray<float> ShortTermLoudnessRelGated;
+		ShortTermLoudnessRelGated.Reserve(ShortTermHistorySize);
+
+		for (uint32 Index = LoudnessRangeStartIndex; Index < LoudnessRangeEndIndex; ++Index)
+		{
+			const float ShortTermEnergy = ShortTermEnergyHistory[Index];
+
+			if (ShortTermEnergy >= RelGateEnergy)
+			{
+				ShortTermLoudnessRelGated.Emplace(FLoudnessAnalyzer::ConvertPerceptualEnergyToLKFS(ShortTermEnergy));
+			}
+		}
+
+		if (ShortTermLoudnessRelGated.IsEmpty())
+		{
+			return NullOpt;
+		}
+
+		// Compute the high and low percentiles of the distribution of values in ShortTermLoudnessRelGated
+		static constexpr int32 LowerPercentile = 10;
+		static constexpr int32 UpperPercentile = 95;
+
+		const float LoudnessRangeMin = LKFSAnalyzerPrivate::GetPercentileValue(ShortTermLoudnessRelGated, LowerPercentile);
+		const float LoudnessRangeMax = LKFSAnalyzerPrivate::GetPercentileValue(ShortTermLoudnessRelGated, UpperPercentile);
+
+		return FFloatInterval(LoudnessRangeMin, LoudnessRangeMax);
+	}
+}
diff --git a/Engine/Source/Runtime/SignalProcessing/Private/LKFSUtils.cpp b/Engine/Source/Runtime/SignalProcessing/Private/LKFSUtils.cpp
new file mode 100644
--- /dev/null
+++ b/Engine/Source/Runtime/SignalProcessing/Private/LKFSUtils.cpp
@@ -0,0 +1,53 @@
+// Copyright Epic Games, Inc. All Rights Reserved.
+
+#include "DSP/LKFSUtils.h"
+
+#include "Math/UnrealMathUtility.h"
+
+namespace Audio::LKFSUtils
+{
+	void TuneSlidingWindowHopSize(int32 InAnalysisHopFrames, int32 InWindowFrames, int32& OutBestWindowHopFrames, int32& OutBestAnalysisHopWindows) 
+	{
+		// The minimum sliding window hop size is the lesser of 
+		// - 25% of the sliding window size
+		// - The analysis hop size
+		int32 MinNumSlidingWindowHopFrames = FMath::Max(1, InWindowFrames / 4);
+		if (MinNumSlidingWindowHopFrames > InAnalysisHopFrames)
+		{
+			MinNumSlidingWindowHopFrames = InAnalysisHopFrames;
+		}
+
+		// The maximum sliding window hop size is the lesser of
+		// - 75% of the sliding window size
+		// - The analysis hop size
+		int32 MaxNumSlidingWindowHopFrames = FMath::Max(1, (3 * InWindowFrames) / 4);
+		if (MaxNumSlidingWindowHopFrames > InAnalysisHopFrames)
+		{
+			MaxNumSlidingWindowHopFrames = InAnalysisHopFrames;
+		}
+
+		// Iterate through the possible sizes to find the one that fits best with our settings.
+		int32 BestWindowHopFrames = -1;
+		int32 BestAnalysisHopWindows = -1;
+		int32 BestMismatch = InAnalysisHopFrames + 2; 
+
+		for (int32 CandidateWindowHopFrames = MinNumSlidingWindowHopFrames; CandidateWindowHopFrames <= MaxNumSlidingWindowHopFrames; CandidateWindowHopFrames++)
+		{
+			int32 CandidateAnalysisHopWindows = InAnalysisHopFrames / CandidateWindowHopFrames;
+			// The match quality is determined by how neatly the sliding window hop size fits into the anlaysis window hop size. 
+			int32 CandidateMismatch = FMath::Abs<int32>(InAnalysisHopFrames - (CandidateWindowHopFrames * CandidateAnalysisHopWindows));
+			if (CandidateMismatch < BestMismatch)
+			{
+				BestMismatch = CandidateMismatch;
+				BestWindowHopFrames = CandidateWindowHopFrames;
+				BestAnalysisHopWindows = CandidateAnalysisHopWindows;
+			}
+		}
+
+		// The algorithm above should always find a valid window hop
+		check(BestWindowHopFrames != -1);
+
+		OutBestWindowHopFrames = BestWindowHopFrames;
+		OutBestAnalysisHopWindows = BestAnalysisHopWindows;
+	}
+} // namespace Audio::LKFSUtils
diff --git a/Engine/Source/Runtime/SignalProcessing/Private/LoudnessAnalyzer.cpp b/Engine/Source/Runtime/SignalProcessing/Private/LoudnessAnalyzer.cpp
new file mode 100644
--- /dev/null
+++ b/Engine/Source/Runtime/SignalProcessing/Private/LoudnessAnalyzer.cpp
@@ -0,0 +1,425 @@
+// Copyright Epic Games, Inc. All Rights Reserved.
+
+#include "DSP/LoudnessAnalyzer.h"
+
+#include "AudioMixer.h"
+#include "DSP/AudioChannelUtils.h"
+#include "DSP/DeinterleaveView.h"
+#include "DSP/FloatArrayMath.h"
+#include "GenericPlatform/GenericPlatformMath.h"
+
+namespace LoudnessAnalyzerPrivate
+{
+	const float Log10Scale = 1.0f / FMath::Loge(10.f);
+}
+
+namespace Audio
+{
+    /* Equal loudness curves from: ANSI Standards S1.4-1983 and S1.42-2001 */
+    /* https://en.wikipedia.org/wiki/A-weighting */
+    float GetEqualLoudnessAWeightForFrequency(const float Freq)
+    {
+        const float FreqSqrd = Freq * Freq;
+
+        return 1.258925411794167f * ((148840000.f * FreqSqrd * FreqSqrd) / ((FreqSqrd + 424.36f) * (FreqSqrd + 148840000.f) * FMath::Sqrt((FreqSqrd + 11599.29f) * (FreqSqrd + 544496.41f))));
+    }
+
+    float GetEqualLoudnessBWeightForFrequency(const float Freq)
+    {
+        const float FreqSqrd = Freq * Freq;
+
+        return 1.019764760044717f * ((148840000.f * Freq * FreqSqrd) / ((FreqSqrd + 424.36f) * (FreqSqrd + 148840000.f) * FMath::Sqrt(FreqSqrd + 25122.25f)));
+    }
+
+    float GetEqualLoudnessCWeightForFrequency(const float Freq)
+    {
+        const float FreqSqrd = Freq * Freq;
+
+        return 1.006931668851804f * ((148840000.f * FreqSqrd) / ((FreqSqrd + 424.36f) * (FreqSqrd + 148840000.f)));
+    }
+
+    float GetEqualLoudnessDWeightForFrequency(const float Freq)
+    {
+        const float FreqSqrd = Freq * Freq;
+        const float H1 = (1037918.48f - FreqSqrd);
+        const float H2 = (9837328.f - FreqSqrd);
+        const float H = (H1 + 1080768.16f * FreqSqrd) / (H2 + 11723776.f * FreqSqrd);
+		return (Freq / 6.8966888496476e-5f) * FMath::Sqrt(H / ((FreqSqrd + 79919.29f) * (FreqSqrd + 1345600.f)));
+    }
+
+	float CalculatePolynomial(const float Freq, const float x0, const float x1, const float x2)
+    {
+	    const double CosFreq = FMath::Cos(Freq);
+    	return x1 * x1 + (x0 - x2) * (x0 - x2) + 2.f * x1 * (x0 + x2) * CosFreq + 4.f * x0 * x2 * CosFreq * CosFreq;
+    }
+
+	float GetEqualLoudnessKWeightForFrequency(const float Freq)
+    {
+	    constexpr float OriginalSampleRateOfKFilter = 48000.f;
+    	constexpr float Nyquist = OriginalSampleRateOfKFilter / 2.f;
+
+    	if (Freq > Nyquist)
+    	{
+    		/**
+    		 * 24kHz is not audible and does not contribute to perceived loudness.
+    		 * This may deviate from other implementations because ITU-R BS.1770-5
+    		 * does not prescribe how to handle higher sampling rates.
+    		 */
+    		return 0.f;
+    	}
+    	if (Freq < 0.f)
+    	{
+    		return 0.f;
+    	}
+
+    	const float NormalizedFrequency = Freq / OriginalSampleRateOfKFilter * 2.f * PI;
+
+    	float PreFilterPower = 0.f;
+	    {
+	    	constexpr float a0 =  1.53512485958697f;
+	    	constexpr float a1 = -2.69169618940638f;
+	    	constexpr float a2 =  1.19839281085285f;
+	    	constexpr float b0 =  1.f;
+	    	constexpr float b1 = -1.69065929318241f;
+	    	constexpr float b2 =  0.73248077421585f;
+
+	    	const float Numerator = CalculatePolynomial(NormalizedFrequency, a0, a1, a2);
+	    	const float Denominator = CalculatePolynomial(NormalizedFrequency, b0, b1, b2);
+
+	    	PreFilterPower = Denominator != 0.f ? Numerator / Denominator : 0.f;
+	    }
+
+    	float RLBFilterPower = 0.f;
+	    {
+	    	constexpr float a0 =  1.f;
+	    	constexpr float a1 = -2.f;
+	    	constexpr float a2 =  1.f;
+	    	constexpr float b0 =  1.f;
+	    	constexpr float b1 = -1.99004745483398f;
+	    	constexpr float b2 =  0.99007225036621f;
+
+	    	const float Numerator = CalculatePolynomial(NormalizedFrequency, a0, a1, a2);
+	    	const float Denominator = CalculatePolynomial(NormalizedFrequency, b0, b1, b2);
+
+	    	RLBFilterPower = Denominator != 0.f ? Numerator / Denominator : 0.f;
+	    }
+
+    	// Convert power to magnitude
+    	return FMath::Sqrt(FMath::Abs(PreFilterPower * RLBFilterPower));
+    }
+
+    float GetEqualLoudnessNoneWeightForFrequency(const float Freq)
+    {
+        return 1.0f;
+    }
+    
+
+	/****************************************************************************/
+	/***********************    FLoudnessAnalyzer    ****************************/
+	/****************************************************************************/
+
+    FLoudnessAnalyzer::FLoudnessAnalyzer(float InSampleRate, const FLoudnessAnalyzerSettings& InSettings)
+    :   Settings(InSettings)
+    {
+		checkf(FMath::IsPowerOfTwo(Settings.FFTSize), TEXT("FFT size must be a power of two"));
+		checkf(Settings.FFTSize >= 8, TEXT("FFT size must be atleast 8"))
+        check(InSampleRate > 0);
+		
+
+		// Create FFT
+		FFFTSettings FFTSettings;
+		FFTSettings.Log2Size = CeilLog2(Settings.FFTSize);
+		FFTSettings.bArrays128BitAligned = true;
+		FFTSettings.bEnableHardwareAcceleration = true;
+
+		ActualFFTSize = 1 << FFTSettings.Log2Size;
+
+		checkf(FFFTFactory::AreFFTSettingsSupported(FFTSettings), TEXT("No fft algorithm supports fft settings."));
+		FFT = FFFTFactory::NewFFTAlgorithm(FFTSettings);
+		if (FFT.IsValid())
+		{
+			checkf(Settings.WindowSize <= ActualFFTSize, TEXT("WindowSize %d must be less than or equal to FFTSize %d"), Settings.WindowSize, ActualFFTSize);
+
+			if (Settings.WindowSize <= 0)
+			{
+				Settings.WindowSize = ActualFFTSize;
+			}
+		
+			// Initialize the analysis window
+			Window.AddUninitialized(Settings.WindowSize);
+			GenerateWindow(Window.GetData(), Settings.WindowType, Window.Num(), 1 /* NumChannels */, false /* bIsPeriodic */);
+			
+			// Size internal buffers
+			WindowedSamples.AddZeroed(FFT->NumInputFloats()); // Zero samples to apply zero buffer in case actual FFT is larger than provided settings.
+			ComplexSpectrum.AddUninitialized(FFT->NumOutputFloats());
+			RealSpectrum.AddUninitialized(ComplexSpectrum.Num() / 2);
+
+			FFTFreqSize = RealSpectrum.Num();
+
+			// Determine which fft bins will be used for calculating loudness
+			const float FreqPerBin = InSampleRate / FFT->NumOutputFloats();
+			MinFreqIndex = static_cast<int32>(FMath::RoundHalfFromZero(Settings.MinAnalysisFrequency / FreqPerBin));
+			MinFreqIndex = FMath::Max(0, MinFreqIndex);
+			MaxFreqIndex = static_cast<int32>(FMath::RoundHalfFromZero(Settings.MaxAnalysisFrequency / FreqPerBin)) + 1;
+			MaxFreqIndex = FMath::Min(MaxFreqIndex, FFTFreqSize);
+
+			CurveWeights.Reset(FFTFreqSize);
+			CurveWeights.AddUninitialized(FFTFreqSize);
+
+			float* CurveWeightsPtr = CurveWeights.GetData();
+
+			const float FFTSize = static_cast<float>(ActualFFTSize);
+
+			for (int32 i = 0; i < FFTFreqSize; i++)
+			{
+				const float Freq = static_cast<float>(i) * InSampleRate / FFTSize;
+				const float Weighting = GetEqualLoudnessWeightForFrequency(Freq);
+
+				// Curve designed for magnitude domain, but applied to power spectrum. The curve is squared to be applied to power spectrum.
+				CurveWeightsPtr[i] = Weighting * Weighting;
+			}
+
+			EnergyScale = CalculateEnergyScaling();
+		}
+    }
+
+    const FLoudnessAnalyzerSettings& FLoudnessAnalyzer::GetSettings() const
+    {
+        return Settings;
+    }
+
+    float FLoudnessAnalyzer::CalculatePerceptualEnergy(TArrayView<const float> InView)
+    {
+		check(InView.Num() == Window.Num());
+
+		float Total = 0.f;
+
+		if (FFT.IsValid())
+		{
+			// Copy input samples and apply window
+			ArrayMultiply(InView, Window, TArrayView<float>(WindowedSamples.GetData(), Window.Num()));
+
+			FFT->ForwardRealToComplex(WindowedSamples.GetData(), ComplexSpectrum.GetData());
+
+			ArrayComplexToPower(ComplexSpectrum, RealSpectrum);
+		
+			const float* PowerSpectrumPtr = RealSpectrum.GetData();
+			const float* CurveWeightsPtr = CurveWeights.GetData();
+			for (int32 i = MinFreqIndex; i < MaxFreqIndex; i++)
+			{
+				Total += PowerSpectrumPtr[i] * CurveWeightsPtr[i];
+			}
+
+			// Normalize by FFT Scaling and by number of samples to calculate mean squared value.
+			Total *= EnergyScale;
+		}
+
+        return Total;
+    }
+
+    float FLoudnessAnalyzer::CalculateLoudness(TArrayView<const float> InView)
+    {
+        float Energy = CalculatePerceptualEnergy(InView);
+		return ConvertPerceptualEnergyToLoudness(Energy);
+    }
+
+	float FLoudnessAnalyzer::ConvertPerceptualEnergyToLoudness(float InPerceptualEnergy) 
+	{
+        return 10.0f * FMath::Loge(InPerceptualEnergy) * LoudnessAnalyzerPrivate::Log10Scale;
+	}
+
+	float FLoudnessAnalyzer::ConvertPerceptualEnergyToLKFS(float InPerceptualEnergy) 
+	{
+        return (10.0f * FMath::Loge(InPerceptualEnergy) * LoudnessAnalyzerPrivate::Log10Scale) - 0.691f;
+	}
+
+	float FLoudnessAnalyzer::GetEqualLoudnessWeightForFrequency(float InHz) const
+	{
+		switch (Settings.LoudnessCurveType)
+		{
+			case ELoudnessCurveType::A:
+				return GetEqualLoudnessAWeightForFrequency(InHz);
+
+			case ELoudnessCurveType::B:
+				return GetEqualLoudnessBWeightForFrequency(InHz);
+
+			case ELoudnessCurveType::C:
+				return GetEqualLoudnessCWeightForFrequency(InHz);
+
+			case ELoudnessCurveType::D:
+				return GetEqualLoudnessDWeightForFrequency(InHz);
+
+			case ELoudnessCurveType::K:
+				return GetEqualLoudnessKWeightForFrequency(InHz);
+
+			case ELoudnessCurveType::None:
+				return GetEqualLoudnessNoneWeightForFrequency(InHz);
+
+			default:
+				{
+					checkNoEntry();
+				}
+				return GetEqualLoudnessNoneWeightForFrequency(InHz);
+		}
+	}
+
+	float FLoudnessAnalyzer::CalculateEnergyScaling() const
+	{
+		switch (Settings.ScalingMethod)
+		{
+			case ELoudnessAnalyzerScalingMethod::BackCompatible:
+				return CalculateEnergyScalingBackCompatible();
+
+			case ELoudnessAnalyzerScalingMethod::Corrected:
+				return CalculateEnergyScalingCorrected();
+
+			default:
+				{
+					checkNoEntry(); // Unimplemented normalization method
+				}
+				return CalculateEnergyScalingCorrected();
+		};
+	}
+
+	float FLoudnessAnalyzer::CalculateEnergyScalingBackCompatible() const
+	{
+		check(FFT.IsValid()); // FFT must be initialized before calling this function 
+
+		const float FFTScaling = GetPowerSpectrumScaling(FFT->Size(), FFT->ForwardScaling(), EFFTScaling::None);
+		const float SizeScaling = 1.f / static_cast<float>(Settings.FFTSize);
+		const float CurveAt1Khz = GetEqualLoudnessWeightForFrequency(1000.f);
+		float CurveScaling = 1.f;
+
+		if (CurveAt1Khz > 0.f)
+		{
+			CurveScaling = 1.f / CurveAt1Khz; // Normalize curve to be unity at 1khz
+			CurveScaling = CurveScaling * CurveScaling; // Convert amplitude normalization to power normalization
+		}
+
+		return FFTScaling * SizeScaling * CurveScaling;
+	}
+
+	float FLoudnessAnalyzer::CalculateEnergyScalingCorrected() const
+	{
+		check(FFT.IsValid()); // FFT must be initialized before calling this function 
+
+		constexpr float HalfSpectrumScaling = 2.f; // When accumulating spectral energy, we ignore the negative spectrum
+												   // magnitudes because they are symmetric with positive spectrum
+												   // magnitudes. This results in accumulation of only half of the 
+												   // spectrum. This scaling corrects that by doubling the energy
+												   // after it has been accumulated.
+		const float FFTScaling = GetPowerSpectrumScaling(FFT->Size(), FFT->ForwardScaling(), EFFTScaling::None);
+		const float WindowScaling = CalculateWindowEnergyScaling();
+
+		return HalfSpectrumScaling * FFTScaling * WindowScaling;
+	}
+
+	float FLoudnessAnalyzer::CalculateWindowEnergyScaling() const
+	{
+		check(Window.Num() > 0); // Window must be initialized before calling this function.
+
+		// Calculate energy of window
+		float WindowEnergy = 0.f;
+		for (float Val : Window)
+		{
+			WindowEnergy += (Val * Val);
+		}
+
+		// Avoid divide by zero on empty window func
+		if (WindowEnergy == 0.f)
+		{
+			return 1.f;
+		}
+		else
+		{
+			return 1.f / WindowEnergy;
+		}
+	}
+
+	/****************************************************************************/
+	/****************    FMultichannelLoudnessAnalyzer    ***********************/
+	/****************************************************************************/
+
+    FMultichannelLoudnessAnalyzer::FMultichannelLoudnessAnalyzer(float InSampleRate, const FLoudnessAnalyzerSettings& InSettings)
+    :   Analyzer(InSampleRate, InSettings)
+    {
+        MonoBuffer.AddUninitialized(Analyzer.GetSettings().WindowSize);
+    }
+
+    float FMultichannelLoudnessAnalyzer::CalculatePerceptualEnergy(TArrayView<const float> InView, const int32 InNumChannels, TArray<float>& OutChannelEnergies)
+    {
+		check(InView.Num() == (InNumChannels * Analyzer.GetSettings().WindowSize));
+
+        OutChannelEnergies.Reset(InNumChannels);
+        OutChannelEnergies.AddUninitialized(InNumChannels);
+        return CalculatePerceptualEnergy(InView, InNumChannels, TArrayView<float>(OutChannelEnergies));
+    }
+
+    float FMultichannelLoudnessAnalyzer::CalculatePerceptualEnergy(TArrayView<const float> InView, const int32 InNumChannels, TArrayView<float> OutChannelEnergies)
+    {
+		check(InView.Num() == (InNumChannels * Analyzer.GetSettings().WindowSize));
+		check(OutChannelEnergies.Num() == InNumChannels);
+		if (ChannelWeights.Num() != InNumChannels)
+		{
+			InitializeChannelWeights(InNumChannels);
+		}
+
+        // Deinterleave audio and send to mono loudness analyzer.
+		TAutoDeinterleaveView<float, FAudioBufferAlignedAllocator> DeinterleaveView(InView, MonoBuffer, InNumChannels);
+		for (TAutoDeinterleaveView<float>::TChannel<FAudioBufferAlignedAllocator> Channel : DeinterleaveView)
+		{
+            OutChannelEnergies[Channel.ChannelIndex] = Analyzer.CalculatePerceptualEnergy(Channel.Values);
+		}
+				
+        // Combine channel energies into overall perceptual energy
+        float Energy = 0.0f;
+        for (int32 ChannelNum = 0; ChannelNum < InNumChannels; ChannelNum++)
+        {
+            Energy += ChannelWeights[ChannelNum] * OutChannelEnergies[ChannelNum];
+        }
+
+        return Energy;
+    }
+
+    float FMultichannelLoudnessAnalyzer::CalculateLoudness(TArrayView<const float> InView, const int32 InNumChannels, TArray<float>& OutChannelLoudness)
+    {
+		OutChannelLoudness.Reset();
+		OutChannelLoudness.AddUninitialized(InNumChannels);
+		return CalculateLoudness(InView, InNumChannels, TArrayView<float>(OutChannelLoudness));
+    }
+
+    float FMultichannelLoudnessAnalyzer::CalculateLoudness(TArrayView<const float> InView, const int32 InNumChannels, TArrayView<float> OutChannelLoudness)
+    {
+        // Convert perceptual energy to loudness 
+        float Energy = CalculatePerceptualEnergy(InView, InNumChannels, OutChannelLoudness);
+        for (float& Value : OutChannelLoudness)
+        {
+            Value = FLoudnessAnalyzer::ConvertPerceptualEnergyToLoudness(Value);
+        }
+        return FLoudnessAnalyzer::ConvertPerceptualEnergyToLoudness(Energy);
+    }
+
+	void FMultichannelLoudnessAnalyzer::InitializeChannelWeights(int32 InNumChannels)
+	{
+		ChannelWeights.SetNum(InNumChannels);
+		for (int32 ChannelIdx = 0; ChannelIdx < InNumChannels; ChannelIdx++)
+		{
+			EAudioMixerChannel::Type ChannelType = EAudioMixerChannel::DefaultChannel;
+
+			if (ensure(ChannelUtils::GetChannelTypeAtIndex(ChannelIdx, ChannelType)))
+			{
+				ChannelWeights[ChannelIdx] = ChannelUtils::GetChannelWeight(ChannelType);
+			}
+			else
+			{
+				ChannelWeights[ChannelIdx] = 1.f;
+			}
+		}
+	}
+
+    const FLoudnessAnalyzerSettings& FMultichannelLoudnessAnalyzer::GetSettings() const
+    {
+        return Analyzer.GetSettings();
+    }
+
+}
diff --git a/Engine/Source/Runtime/SignalProcessing/Public/DSP/AudioChannelUtils.h b/Engine/Source/Runtime/SignalProcessing/Public/DSP/AudioChannelUtils.h
new file mode 100644
--- /dev/null
+++ b/Engine/Source/Runtime/SignalProcessing/Public/DSP/AudioChannelUtils.h
@@ -0,0 +1,17 @@
+// Copyright Epic Games, Inc. All Rights Reserved.
+#pragma once
+
+#include "AudioMixerChannel.h"
+
+#define UE_API SIGNALPROCESSING_API
+
+namespace Audio::ChannelUtils
+{	
+	/** Helper function to get the channel map type at the given index. */
+	UE_API bool GetChannelTypeAtIndex(const int32 Index, EAudioMixerChannel::Type& OutType);
+	
+	/** Helper function to get the channel weight given a channel type. */
+	UE_API float GetChannelWeight(EAudioMixerChannel::Type InType);
+} // namespace Audio::ChannelUtils
+
+#undef UE_API
diff --git a/Engine/Source/Runtime/SignalProcessing/Public/DSP/LKFSAnalyzer.h b/Engine/Source/Runtime/SignalProcessing/Public/DSP/LKFSAnalyzer.h
new file mode 100644
--- /dev/null
+++ b/Engine/Source/Runtime/SignalProcessing/Public/DSP/LKFSAnalyzer.h
@@ -0,0 +1,224 @@
+// Copyright Epic Games, Inc. All Rights Reserved.
+#pragma once
+
+#include "Containers/CircularBuffer.h"
+#include "DSP/LoudnessAnalyzer.h"
+#include "DSP/SlidingWindow.h"
+#include "HAL/Platform.h"
+
+#define UE_API SIGNALPROCESSING_API
+
+namespace Audio
+{
+	template <typename InSampleType> class TSlidingBuffer;
+
+	/** Settings for FLKFSAnalyzer */
+	struct FLKFSAnalyzerSettings
+	{
+		// These values are taken from the suggested standards for LKFS/LUFS defined in ITU-R BS.1770
+		static constexpr float StandardShortTermLoudnessDuration = 3.f;			// 3 seconds
+		static constexpr float StandardIntegratedLoudnessAnalysisPeriod = 1.f;	// 1 second
+		
+		// Streaming audio does not have a natural end point. Integrated loudness in this 
+		// implementation is taken over a long duration (as opposed to accumulating indefinitely). 
+		static constexpr float DefaultIntegratedLoudnessDuration = 60.f; // 1 minute
+
+		FLoudnessAnalyzerSettings LoudnessAnalyzerSettings {
+			.FFTSize = 2048,
+			.WindowType = EWindowType::Blackman,
+			.LoudnessCurveType = ELoudnessCurveType::K,
+			.MinAnalysisFrequency = 20.0f,
+			.MaxAnalysisFrequency = 24000.0f,
+			.ScalingMethod = ELoudnessAnalyzerScalingMethod::Corrected
+		};
+
+		/** Number of seconds between loudness measurements */
+		float AnalysisPeriod = 0.1f;
+
+		/** Number of seconds of audio analyzed for each loudness measurements */
+		float AnalysisWindowDuration = 0.4f;
+
+		/** Duration of audio analyzed for short term loudness. */
+		float ShortTermLoudnessDuration = StandardShortTermLoudnessDuration;
+
+		/** Number of seconds between long term loudness updates. */
+		float IntegratedLoudnessAnalysisPeriod = StandardIntegratedLoudnessAnalysisPeriod;
+
+		/** Duration of audio analyzed for long term loudness. */
+		float IntegratedLoudnessDuration = DefaultIntegratedLoudnessDuration;
+
+		/** Should the Loudness Range (LRA) statistic be calculated for each channel? */
+		bool bCalculatePerChannelLoudnessRange = false;
+
+		/** Should the overall Loudness Range (LRA) statistic be calculated? */
+		bool bCalculateOverallLoudnessRange = false;
+	};
+
+	/** LKFS Result for a single channel (or overall channels result if Channel is set to -1) */
+	struct FLKFSResult
+	{
+		/** The audio channel index which produced this result. If -1, then the data was produced by combining all other audio channels */
+		int32 Channel = 0;
+
+		/** Time in seconds of the source audio which corresponds to the loudness measurements. */
+		float Timestamp = 0.0f;
+
+		/** The instantaneous, perceptually weighted energy relative to full scale. */
+		float Energy = 0.f;
+
+		/** The instantaneous, perceptually weighted loudness in dB. */
+		float Loudness = 0.0f;
+
+		/** The average loudness over a short term window of time in dB. */
+		float ShortTermLoudness = 0.0f;
+
+		/** The average loudness over a long term window of time in dB. */
+		float IntegratedLoudness = 0.0f;
+
+		/** The average gated loudness over a long term window of time in dB. */
+		float GatedLoudness = 0.0f;
+
+		/** Loudness Range (LRA), as defined in EBU Tech 3342. */
+		TOptional<FFloatInterval> LoudnessRange;
+	};
+
+	/** Analyzer results contains all channels LKFS results */
+	class FLKFSAnalyzerResults
+	{
+	public:
+		/** Construct analyzer results */
+		FLKFSAnalyzerResults() = default;
+
+		/** The number of individual channels available in the result. */
+		UE_API int32 GetNumChannels() const;
+
+		/** Add to the result. */
+		UE_API void Add(FLKFSResult InResults);
+
+		/** Returns const reference to a FLKFSResult array for individual channel. */
+		UE_API const TArray<FLKFSResult>& GetChannelLoudnessResults(int32 ChannelIdx) const;
+
+		/** Returns const reference to a FLKFSResult array associated with overall loudness. */
+		UE_API const TArray<FLKFSResult>& GetLoudnessResults() const;
+
+		/** Denotes the overall loudness channel index as opposed individual channel indices. */
+		inline static const int32 ChannelIndexOverall = -1;
+
+	private:
+		TSortedMap<int32, TArray<FLKFSResult>> ChannelResults;
+	};
+
+
+	/** FLKFSAnalyzer
+	 *
+	 * FLKFSAnalyzer computes Loudness K-Weighted FS (LKFS) metrics for interleaved audio channels.
+	 * Uses a sliding window approach and an internal perceptual energy analyzer.
+	 * Produces instantaneous energy and loudness as well as short term, integrated, gated LKFS
+	 * and optionally the loudness range (based on EBU-3342).
+	 */
+	class FLKFSAnalyzer
+	{
+	public:
+		/** Construct analyzer */
+		UE_API FLKFSAnalyzer(int32 InNumChannels, float InSampleRate, const FLKFSAnalyzerSettings& InSettings);
+
+		/** Return const reference to settings used inside this analyzer. */
+		UE_API const FLKFSAnalyzerSettings& GetSettings() const;
+
+		/** Analyze audio to obtain LKFS values (see FLKFSAnalyzerResults and FLKFSResult for more details) */
+		UE_API void Analyze(TArrayView<const float> InAudio, FLKFSAnalyzerResults& OutLKFSResults);
+
+	private:
+		/**
+		 * Holds aggregated loudness statistics calculated over a history of energy measurements.
+		 */
+		struct FAggregateLoudnessStats
+		{
+			static constexpr float InvalidLoudness = TNumericLimits<float>::Lowest();
+
+			float ShortTermLoudness  = InvalidLoudness;
+			float IntegratedLoudness = InvalidLoudness;
+			float GatedLoudness      = InvalidLoudness;
+
+			TOptional<FFloatInterval> LoudnessRange;
+		};
+
+		/**
+		 * Holds state for Loudness Range calculation (for a single channel).
+		 */
+		class FLoudnessRange
+		{
+		public:
+			explicit FLoudnessRange(uint32 InMaxShortTermHistorySize)
+				: ShortTermEnergyHistory(InMaxShortTermHistorySize)
+				, MaxShortTermHistorySize(InMaxShortTermHistorySize)
+			{
+			}
+
+			void AddShortTermEnergyMeasurement(float InShortTermEnergy);
+
+			TOptional<FFloatInterval> CalculateLoudnessRange() const;
+
+		private:
+			TCircularBuffer<float> ShortTermEnergyHistory;
+
+			uint32 MaxShortTermHistorySize = 0;
+			uint32 ShortTermHistorySize = 0;
+			uint32 ShortTermWriteIndex = 0;
+		};
+
+		void AddCurrentLoudnessDataToResult(FLKFSAnalyzerResults& OutResult);
+		void AnalyzeSlidingWindow(TAutoSlidingWindow<float>& InSlidingWindow, FLKFSAnalyzerResults& OutResult);
+		void AnalyzeWindow(TArrayView<const float> InWindow);
+
+		void AddCurrentLoudnessDataToHistory();
+
+		void UpdateAggregateLoudnessStats();
+		void UpdateAggregateLoudnessStats(const TCircularBuffer<float>& InHistory, int32 InOffset, int32 InStride, bool bInUpdateIntegratedLoudness, FAggregateLoudnessStats& OutStats, FLoudnessRange* LoudnessRange = nullptr) const;
+
+		FLKFSAnalyzerSettings Settings;
+
+		int32 NumChannels = 0;
+		float SampleRate = 0.0f;
+
+		int32 NumAnalyzedBuffers = 0;
+		int32 NumEnergyHistory = 0;
+		int32 NumSlidingWindowHopFrames = 0;
+
+		// Number of windows between each loudness analysis
+		int32 NumAnalysisHopWindows = 1;
+
+		// Number of windows inspected for each loudness analysis
+		int32 NumAnalysisWindows = 1;
+
+		// Number of values in the "EnergyHistory" to analyze for short term loudness
+		int32 NumShortTermHistoryResults = 1;
+
+		// Number of values in the "EnergyHistory" to analyze for integrated and gated loudness
+		int32 NumIntegratedHistoryResults = 1;
+
+		// Number of values in the "EnergyHistory" between integrated and gated loudness anlysis
+		int32 NumIntegratedHopHistoryResults = 1;
+
+		TArray<float> ChannelPerceptualEnergy;
+		TArray<float> OverallPerceptualEnergy;
+		TArray<float> InternalWindow;
+
+		// Per channel stats for short-term, integrated and gated loudness, and optional loudness range. 
+		TSortedMap<int32, FAggregateLoudnessStats> AggregateLoudnessStats;
+		
+		// Per channel state for calculating Loudness Range.
+		TSortedMap<int32, FLoudnessRange> LoudnessRangeStatePerChannel;
+
+		// Channel interleaved history buffer of energy measurements.
+		// These are taken once per an loudness analysis, not necessarily once per a audio window. 
+		TCircularBuffer<float> ChannelEnergyHistory;
+		TCircularBuffer<float> OverallEnergyHistory;
+
+		TSlidingBuffer<float> InternalBuffer;
+
+		FMultichannelLoudnessAnalyzer LoudnessAnalyzer;
+	};
+}
+
+#undef UE_API
diff --git a/Engine/Source/Runtime/SignalProcessing/Public/DSP/LKFSUtils.h b/Engine/Source/Runtime/SignalProcessing/Public/DSP/LKFSUtils.h
new file mode 100644
--- /dev/null
+++ b/Engine/Source/Runtime/SignalProcessing/Public/DSP/LKFSUtils.h
@@ -0,0 +1,17 @@
+// Copyright Epic Games, Inc. All Rights Reserved.
+#pragma once
+
+#include "HAL/Platform.h"
+#include "Misc/CoreMiscDefines.h"
+
+#define UE_API SIGNALPROCESSING_API
+
+namespace Audio::LKFSUtils
+{
+	// There may be multiple calls to the LoudnessAnalyzer to produce a single result. 
+	// The sliding window hop is tuned here so that it best matches the desired AnalysisPeriod 
+	// while maintaining between a 25% to 75% window overlap.
+	UE_INTERNAL UE_API void TuneSlidingWindowHopSize(int32 InNumAnalysisHopFrames, int32 InNumSlidingWindowFrames, int32& InNumSlidingWindowHopFrames, int32& InNumAnalysisHopWindows);
+} // namespace Audio::LKFSUtils
+
+#undef UE_API
diff --git a/Engine/Source/Runtime/SignalProcessing/Public/DSP/LoudnessAnalyzer.h b/Engine/Source/Runtime/SignalProcessing/Public/DSP/LoudnessAnalyzer.h
new file mode 100644
--- /dev/null
+++ b/Engine/Source/Runtime/SignalProcessing/Public/DSP/LoudnessAnalyzer.h
@@ -0,0 +1,237 @@
+// Copyright Epic Games, Inc. All Rights Reserved.
+
+#pragma once
+
+#include "AudioFFT.h"
+
+#define UE_API SIGNALPROCESSING_API
+
+namespace Audio { class IFFTAlgorithm; }
+
+namespace Audio
+{
+    /** Equal loudness curves from: ANSI Standards S1.4-1983 and S1.42-2001 */
+	enum class ELoudnessCurveType : uint8
+	{
+		A,
+		B,
+		C,
+		D,
+		K,
+		None
+	};
+
+	/** Generate equal loudness weight for magnitude spectrum using curve A */
+    UE_API float GetEqualLoudnessAWeigthForFrequency(const float Freq);
+
+	/** Generate equal loudness weight for magnitude spectrum using curve B */
+    UE_API float GetEqualLoudnessBWeigthForFrequency(const float Freq);
+
+	/** Generate equal loudness weight for magnitude spectrum using curve C */
+    UE_API float GetEqualLoudnessCWeigthForFrequency(const float Freq);
+
+	/** Generate equal loudness weight for magnitude spectrum using curve D */
+    UE_API float GetEqualLoudnessDWeigthForFrequency(const float Freq);
+
+	/** Generate equal loudness weight for magnitude spectrum using curve None */
+    UE_API float GetEqualLoudnessNoneWeigthForFrequency(const float Freq);
+
+	enum class ELoudnessAnalyzerScalingMethod : uint8
+	{
+		// Normalize by FFTSize and having perceptual curves be unity gain at 1khz. 
+		// 
+		// Note: This method does not correctly account for different window functions
+		// and also reduces the energy by half. It remains here for backwards
+		// compatibility.
+		BackCompatible,
+
+		// Normalize by Window Function. Use this when performing loudness analysis
+		// which must match the LKFS standardized measurements. 
+		Corrected
+	};
+
+	/** Settings for FLoudnessAnalyzer */
+    struct FLoudnessAnalyzerSettings
+    {
+		/** Size of FFT describe in number of samples */
+        int32 FFTSize = 8192;
+
+		/** Type of window to be applied to input audio */
+        EWindowType WindowType = EWindowType::Blackman;
+
+		/** Type of equal loudness curve to weight power spectrum */
+        ELoudnessCurveType LoudnessCurveType = ELoudnessCurveType::D;
+
+		/** Minimum frequency for analyzing audio */
+        float MinAnalysisFrequency = 20.f;
+
+		/** Maximum frequency for analyzing audio */
+        float MaxAnalysisFrequency = 10000.f;
+
+		/** Size of the audio analysis window in samples.
+		 *
+		 * If this is less than 1, it will become equal to the FFTSize.
+		 * Otherwise, this must be less than or equal to the FFT Size.
+		 */
+		int32 WindowSize = -1;
+
+		/** Method to use for scaling and normalizing data. Corrected is preferred
+		 * in terms of correctness. BackCompatible is available for backwards compatibility.
+		 */
+		ELoudnessAnalyzerScalingMethod ScalingMethod = ELoudnessAnalyzerScalingMethod::BackCompatible;
+    };
+
+	/** FLoudnessAnalyzer
+	 *
+	 * FLoudnessAnalyzer will calculate perceptual energy or loudness one window at a time. 
+	 * This analyzer only handles mono audio. See FMultichannelLoudnessAnalyzer for multichannel
+	 * support.
+	 */
+    class FLoudnessAnalyzer
+    {
+        public:
+			/** Construct analyzer */
+            UE_API FLoudnessAnalyzer(float InSampleRate, const FLoudnessAnalyzerSettings& InSettings);
+
+			/**
+			 * Calculate the perceptual energy for the input samples.  The number of 
+			 * input samples must be equal to the FFTSize supplied in FLoudnessAnalyzerSettings.
+			 */
+            UE_API float CalculatePerceptualEnergy(TArrayView<const float> InSampleView);
+
+			/**
+			 * Calculate the loudness for the input samples.  The number of  input samples must 
+			 * be equal to the FFTSize supplied in FLoudnessAnalyzerSettings.
+			 */
+            UE_API float CalculateLoudness(TArrayView<const float> InSampleView);
+
+			/**
+			 * Convert from perceptual energy to loudness [Loudness = 10 * log_10(Energy)]
+			 */
+			static UE_API float ConvertPerceptualEnergyToLoudness(float InPerceptualEnergy);
+			
+			/**
+			 * Convert from perceptual energy to loudness [LKFS = -0.691 + 10 * log_10(Energy)]
+			 */
+			static UE_API float ConvertPerceptualEnergyToLKFS(float InPerceptualEnergy);
+
+			/**
+			 * Return const reference to settings used inside this analyzer.
+			 */
+            UE_API const FLoudnessAnalyzerSettings& GetSettings() const;
+			
+        private:
+			float CalculateEnergyScaling() const;
+			float GetEqualLoudnessWeightForFrequency(float InHz) const;
+			float CalculateEnergyScalingBackCompatible() const;
+			float CalculateEnergyScalingCorrected() const;
+			float CalculateWindowEnergyScaling() const;
+
+            FLoudnessAnalyzerSettings Settings;
+
+			/** Scale factor used to normalize between differing window sizes and fft sizes */
+            float EnergyScale = 1.f;
+
+            int32 MinFreqIndex = 0;
+            int32 MaxFreqIndex = 0;
+
+			/** Number of FFTFrequency bins worth looking at. */
+			int32 FFTFreqSize = 0;
+			int32 ActualFFTSize = 0;
+
+			TUniquePtr<IFFTAlgorithm> FFT;
+
+			FAlignedFloatBuffer Window;
+            FAlignedFloatBuffer WindowedSamples;
+            FAlignedFloatBuffer ComplexSpectrum;
+            FAlignedFloatBuffer RealSpectrum;
+
+            FAlignedFloatBuffer CurveWeights;
+    };
+
+
+	/** FMultichannelLoudnessAnalyzer
+	 *
+	 * FMultichannelLoudnessAnalyzer will calculate perceptual energy or loudness one window 
+	 * at a time. This module can handle one or more channels at a time. It expects 
+	 * interleaved audio samples.
+	 */
+    class FMultichannelLoudnessAnalyzer 
+    {
+        public:
+			/** Construct analyzer. */
+            UE_API FMultichannelLoudnessAnalyzer(float InSampleRate, const FLoudnessAnalyzerSettings& InSettings);
+
+			/** 
+			 * Calculate perceptual energy. 
+			 *
+			 * InSampleView		  InSampleView is expected to contain interleaved audio with 
+			 *                    (InSettings.FFTSize * InNumChannels) number of samples.
+			 *
+			 * InNumChannels      Number of input channels.
+			 *
+			 * OutChannelEnergies Will be populated with energies on a per channel basis.
+			 *
+			 * The overall perceptual energy is returned.
+			 */
+			UE_DEPRECATED(5.7, "Use method which accepts a TArrayView<float> for OutChannelEnergies")
+            UE_API float CalculatePerceptualEnergy(TArrayView<const float> InSampleView, const int32 InNumChannels, TArray<float>& OutChannelEnergies); 
+
+			/** 
+			 * Calculate perceptual energy. 
+			 *
+			 * InSampleView		  InSampleView is expected to contain interleaved audio with 
+			 *                    (InSettings.FFTSize * InNumChannels) number of samples.
+			 *
+			 * InNumChannels      Number of input channels.
+			 *
+			 * OutChannelEnergies Will be populated with energies on a per channel basis.
+			 *
+			 * The overall perceptual energy is returned.
+			 */
+            UE_API float CalculatePerceptualEnergy(TArrayView<const float> InSampleView, const int32 InNumChannels, TArrayView<float> OutChannelEnergies); 
+
+			/** 
+			 * Calculate loudness
+			 *
+			 * InSampleView		  InSampleView is expected to contain interleaved audio with 
+			 *                    (InSettings.FFTSize * InNumChannels) number of samples.
+			 *
+			 * InNumChannels      Number of input channels.
+			 *
+			 * OutChannelLoudness Will be populated with energies on a per channel basis.
+			 *
+			 * The overall loudness is returned.
+			 */
+			UE_DEPRECATED(5.7, "Use method which accepts a TArrayView<float> for OutChannelLoudness")
+            UE_API float CalculateLoudness(TArrayView<const float> InSampleView, const int32 InNumChannels, TArray<float>& OutChannelLoudness);
+
+			/** 
+			 * Calculate loudness
+			 *
+			 * InSampleView		  InSampleView is expected to contain interleaved audio with 
+			 *                    (InSettings.FFTSize * InNumChannels) number of samples.
+			 *
+			 * InNumChannels      Number of input channels.
+			 *
+			 * OutChannelLoudness Will be populated with energies on a per channel basis.
+			 *
+			 * The overall loudness is returned.
+			 */
+            UE_API float CalculateLoudness(TArrayView<const float> InSampleView, const int32 InNumChannels, TArrayView<float> OutChannelLoudness);
+
+			/**
+			 * Return const reference to settings used inside this analyzer.
+			 */
+            UE_API const FLoudnessAnalyzerSettings& GetSettings() const;
+
+        private:
+			void InitializeChannelWeights(int32 InNumChannels);
+
+            FAlignedFloatBuffer MonoBuffer;
+            TArray<float> ChannelWeights;
+            FLoudnessAnalyzer Analyzer;
+    };
+}
+
+#undef UE_API
