Fix swapped arguments to CrossCorrelate in YIN pitch detector that caused truncated autocorrelation sums. The incorrect argument order lost up to t terms from each lag's sum, causing pitch detection failures
at higher sample rates and tight frequency ranges.

#rnx
#Synced-CL 51850588

diff --git a/Engine/Plugins/Runtime/AudioSynesthesia/Source/AudioSynesthesiaCore/Private/YINPitchDetector.cpp b/Engine/Plugins/Runtime/AudioSynesthesia/Source/AudioSynesthesiaCore/Private/YINPitchDetector.cpp
--- a/Engine/Plugins/Runtime/AudioSynesthesia/Source/AudioSynesthesiaCore/Private/YINPitchDetector.cpp
+++ b/Engine/Plugins/Runtime/AudioSynesthesia/Source/AudioSynesthesiaCore/Private/YINPitchDetector.cpp
@@ -2,10 +2,12 @@
 
 #include "YINPitchDetector.h"
 
-#include "AudioSynesthesiaCoreLog.h"
 #include "DSP/BlockCorrelator.h"
 #include "DSP/Dsp.h"
 #include "ProfilingDebugging/CpuProfilerTrace.h"
+#include "Logging/LogMacros.h"
+
+DEFINE_LOG_CATEGORY_STATIC(LogYINPitchDetector, Log, All);
 
 namespace Audio
 {
@@ -50,7 +52,7 @@ namespace Audio
 
 		if (WindowNumSamples > MaximumWindowSize)
 		{
-			UE_LOG(LogAudioSynesthesiaCore, Warning,
+			UE_LOG(LogYINPitchDetector, Warning,
 				TEXT("YIN: Pitches at minimum frequency (%.1f Hz) may not be detected reliably due to internal window size limit (%d samples)"),
 				Settings.MinimumFrequency, MaximumWindowSize);
 			WindowNumSamples = MaximumWindowSize;
@@ -66,7 +68,7 @@ namespace Audio
 
 		if (MinLagBin >= MaxLagBin)
 		{
-			UE_LOG(LogAudioSynesthesiaCore, Warning,
+			UE_LOG(LogYINPitchDetector, Warning,
 				TEXT("YIN: Frequency range [%.1f, %.1f] Hz maps to a degenerate lag range [%d, %d]. Check settings."),
 				Settings.MinimumFrequency, Settings.MaximumFrequency, MinLagBin, MaxLagBin);
 		}
@@ -136,10 +138,11 @@ namespace Audio
 		// The second half remains permanently zeroed (set in constructor).
 		FMemory::Memcpy(FirstHalfBuffer.GetData(), Data, HalfWindowNumSamples * sizeof(float));
 
-		// CrossCorrelate(A, B)[τ] = Σ_n A[n] * B[n+τ]
-		// With A = FirstHalfBuffer (x[0..W/2-1] + zeros) and B = InWindow (x[0..W-1]):
+		// CrossCorrelate(A, B) internally computes IFFT(conj(FFT(B)) * FFT(A)),
+		// giving Output[τ] = Σ_n B[n] * A[n+τ].
+		// With B = FirstHalfBuffer (x[0..W/2-1] + zeros) and A = InWindow (x[0..W-1]):
 		//   Output[τ] = Σ_{t=0}^{W/2-1} x[t] * x[t+τ]  = r(τ)  for τ ≤ W/2
-		Correlator->CrossCorrelate(FirstHalfBuffer, InWindow, AutoCorrBuffer);
+		Correlator->CrossCorrelate(InWindow, FirstHalfBuffer, AutoCorrBuffer);
 
 		const float* RESTRICT AutoCorr = AutoCorrBuffer.GetData();
 
