Add bounds check to skip processing when MaxBin is negative in the auto-correlation pitch detector

#rnx

#Synced-CL 52395192

diff --git a/Engine/Plugins/Runtime/AudioSynesthesia/Source/AudioSynesthesiaCore/Private/AutoCorrelationPitchDetector.cpp b/Engine/Plugins/Runtime/AudioSynesthesia/Source/AudioSynesthesiaCore/Private/AutoCorrelationPitchDetector.cpp
--- a/Engine/Plugins/Runtime/AudioSynesthesia/Source/AudioSynesthesiaCore/Private/AutoCorrelationPitchDetector.cpp
+++ b/Engine/Plugins/Runtime/AudioSynesthesia/Source/AudioSynesthesiaCore/Private/AutoCorrelationPitchDetector.cpp
@@ -154,6 +154,10 @@ namespace Audio
 				float Energy = FMath::Max(SMALL_NUMBER, AutoCorrBuffer[0]);
 
 				MaxBin = FMath::Min(AutoCorrBuffer.Num() - MinAutoCorrBin - 1, MaxAutoCorrBin - MinAutoCorrBin);
+				if (MaxBin < 1)
+				{
+					continue;
+				}
 
 				TArrayView<float> AutoCorrView(&AutoCorrBuffer.GetData()[MinAutoCorrBin], MaxBin + 1);
 				PeakPicker->PickPeaks(AutoCorrView, PeakIndices);
